New experimental Mesh Resizing plugin (no actual code, just some boilerplate).

diff --git a/Engine/Plugins/Experimental/MeshResizing/MeshResizing.uplugin b/Engine/Plugins/Experimental/MeshResizing/MeshResizing.uplugin
new file mode 100644
--- /dev/null
+++ b/Engine/Plugins/Experimental/MeshResizing/MeshResizing.uplugin
@@ -0,0 +1,34 @@
+{
+    "FileVersion" : 3,
+    "Version" : 1,
+    "VersionName" : "0.1",
+    "FriendlyName" : "Mesh Resizing",
+    "Description" : "Mesh Resizing",
+    "Category": "Editor",
+    "CreatedBy" : "Epic Games, Inc.",
+    "CreatedByURL" : "https://epicgames.com",
+    "DocsURL" : "",
+    "MarketplaceURL" : "",
+    "SupportURL" : "",
+    "EnabledByDefault" : false,
+    "CanContainContent" : true,
+	"IsExperimentalVersion": true,
+    "Installed" : false,
+	"Modules": [
+		{
+			"Name": "MeshResizingDataflowNodes",
+			"Type": "Runtime",
+			"LoadingPhase": "Default"
+		}
+	],
+	"Plugins": [
+		{
+			"Name": "Dataflow",
+			"Enabled": true
+		},
+		{
+			"Name": "GeometryCollectionPlugin",
+			"Enabled": true
+		}
+	]
+}
diff --git a/Engine/Plugins/Experimental/MeshResizing/Source/MeshResizingCore/MeshResizingCore.Build.cs b/Engine/Plugins/Experimental/MeshResizing/Source/MeshResizingCore/MeshResizingCore.Build.cs
new file mode 100644
--- /dev/null
+++ b/Engine/Plugins/Experimental/MeshResizing/Source/MeshResizingCore/MeshResizingCore.Build.cs
@@ -0,0 +1,18 @@
+// Copyright Epic Games, Inc. All Rights Reserved.
+
+namespace UnrealBuildTool.Rules
+{
+	public class MeshResizingCore : ModuleRules
+	{
+        public MeshResizingCore(ReadOnlyTargetRules Target) : base(Target)
+		{
+			PCHUsage = ModuleRules.PCHUsageMode.UseExplicitOrSharedPCHs;
+
+			PrivateDependencyModuleNames.AddRange(
+				new string[]
+				{
+				}
+			);
+		}
+	}
+}
diff --git a/Engine/Plugins/Experimental/MeshResizing/Source/MeshResizingEngine/MeshResizingEngine.Build.cs b/Engine/Plugins/Experimental/MeshResizing/Source/MeshResizingEngine/MeshResizingEngine.Build.cs
new file mode 100644
--- /dev/null
+++ b/Engine/Plugins/Experimental/MeshResizing/Source/MeshResizingEngine/MeshResizingEngine.Build.cs
@@ -0,0 +1,22 @@
+// Copyright Epic Games, Inc. All Rights Reserved.
+
+namespace UnrealBuildTool.Rules
+{
+	public class MeshResizingEngine : ModuleRules
+	{
+		public MeshResizingEngine(ReadOnlyTargetRules Target) : base(Target)
+		{
+			bTreatAsEngineModule = true;
+			PublicDependencyModuleNames.AddRange(
+				new string[]
+				{
+				}
+			);
+			PrivateDependencyModuleNames.AddRange(
+				new string[]
+				{	
+				}
+			);
+		}
+	}
+}
diff --git a/Engine/Plugins/Experimental/MeshResizing/Source/MeshResizingNodes/MeshResizingDataflowNodes.Build.cs b/Engine/Plugins/Experimental/MeshResizing/Source/MeshResizingNodes/MeshResizingDataflowNodes.Build.cs
new file mode 100644
--- /dev/null
+++ b/Engine/Plugins/Experimental/MeshResizing/Source/MeshResizingNodes/MeshResizingDataflowNodes.Build.cs
@@ -0,0 +1,34 @@
+// Copyright Epic Games, Inc. All Rights Reserved.
+
+using UnrealBuildTool;
+public class MeshResizingDataflowNodes : ModuleRules
+{
+	public MeshResizingDataflowNodes(ReadOnlyTargetRules Target) : base(Target)
+	{
+		PublicIncludePathModuleNames.AddRange(
+			new string[]
+			{
+			}
+		);
+
+		PublicDependencyModuleNames.AddRange(
+			new string[]
+			{
+				"Core",
+				"CoreUObject",
+				"Engine",
+				"GeometryCore",
+			}
+		);
+		PrivateDependencyModuleNames.AddRange(
+			new string[]
+			{
+				"DataflowCore",
+				"DataflowEngine",
+				"DataflowNodes",
+				"GeometryFramework",
+				"MeshConversionEngineTypes"
+			}
+			);
+	}
+}
diff --git a/Engine/Plugins/Experimental/MeshResizing/Source/MeshResizingNodes/Private/MeshResizing/BaseBodyDataflowNodes.cpp b/Engine/Plugins/Experimental/MeshResizing/Source/MeshResizingNodes/Private/MeshResizing/BaseBodyDataflowNodes.cpp
new file mode 100644
--- /dev/null
+++ b/Engine/Plugins/Experimental/MeshResizing/Source/MeshResizingNodes/Private/MeshResizing/BaseBodyDataflowNodes.cpp
@@ -0,0 +1,52 @@
+// Copyright Epic Games, Inc. All Rights Reserved.
+
+#include "MeshResizing/BaseBodyDataflowNodes.h"
+#include "Dataflow/DataflowCore.h"
+#include "DynamicMesh/DynamicMesh3.h"
+#include "Engine/SkeletalMesh.h"
+#include "Rendering/SkeletalMeshRenderData.h"
+#include "SkeletalMeshLODRenderDataToDynamicMesh.h"
+#include "UDynamicMesh.h"
+
+#include UE_INLINE_GENERATED_CPP_BY_NAME(BaseBodyDataflowNodes)
+
+namespace UE::MeshResizing
+{
+	void RegisterBaseBodyDataflowNodes()
+	{
+		static const FLinearColor CDefaultNodeBodyTintColor = FLinearColor(0.f, 0.f, 0.f, 0.5f);
+
+		DATAFLOW_NODE_REGISTER_CREATION_FACTORY(FSkeletalMeshToMeshDataflowNode);
+	}
+}
+
+void FSkeletalMeshToMeshDataflowNode::Evaluate(UE::Dataflow::FContext& Context, const FDataflowOutput* Out) const
+{
+	if (Out->IsA(&Mesh))
+	{
+		if (const USkeletalMesh* const InSkeletalMesh = GetValue(Context, &SkeletalMesh))
+		{
+			if (const FSkeletalMeshRenderData* const RenderData = InSkeletalMesh->GetResourceForRendering())
+			{
+				if (RenderData->LODRenderData.IsValidIndex(LODLevel))
+				{
+					const FSkeletalMeshLODRenderData* const SkeletalMeshLODRenderData = &RenderData->LODRenderData[LODLevel];
+
+					TObjectPtr<UDynamicMesh> NewMesh = NewObject<UDynamicMesh>();
+					NewMesh->Reset();
+
+					UE::Geometry::FDynamicMesh3& DynMesh = NewMesh->GetMeshRef();
+
+					UE::Geometry::FSkeletalMeshLODRenderDataToDynamicMesh::ConversionOptions  ConversionOptions;
+					//ConversionOptions.bWantTangents = AssetOptions.bRequestTangents;
+
+					UE::Geometry::FSkeletalMeshLODRenderDataToDynamicMesh::Convert(SkeletalMeshLODRenderData, InSkeletalMesh->GetRefSkeleton(), ConversionOptions, DynMesh);
+
+					SetValue(Context, NewMesh, &Mesh);
+					return;
+				}
+			}
+		}
+		SetValue(Context, TObjectPtr<UDynamicMesh>(NewObject<UDynamicMesh>()), &Mesh);
+	}
+}
diff --git a/Engine/Plugins/Experimental/MeshResizing/Source/MeshResizingNodes/Private/MeshResizing/MeshResizingDataflowNodesModule.cpp b/Engine/Plugins/Experimental/MeshResizing/Source/MeshResizingNodes/Private/MeshResizing/MeshResizingDataflowNodesModule.cpp
new file mode 100644
--- /dev/null
+++ b/Engine/Plugins/Experimental/MeshResizing/Source/MeshResizingNodes/Private/MeshResizing/MeshResizingDataflowNodesModule.cpp
@@ -0,0 +1,42 @@
+// Copyright Epic Games, Inc. All Rights Reserved.
+
+#include "Dataflow/DataflowNodeColorsRegistry.h"
+#include "Dataflow/DataflowNodeFactory.h"
+#include "Modules/ModuleManager.h"
+#include "MeshResizing/BaseBodyDataflowNodes.h"
+
+namespace UE::MeshResizing
+{
+	namespace Private
+	{
+		struct FColorScheme
+		{
+			static inline const FLinearColor Asset = FColor(180, 120, 110);
+			static inline const FLinearColor NodeHeader = FColor(180, 120, 110);
+			static inline const FLinearColor NodeBody = FColor(18, 12, 11, 127);
+		};
+
+		static void RegisterDataflowNodes()
+		{
+			DATAFLOW_NODE_REGISTER_CREATION_FACTORY_NODE_COLORS_BY_CATEGORY("MeshResizing", FColorScheme::NodeHeader, FColorScheme::NodeBody);
+
+			RegisterBaseBodyDataflowNodes();
+		}
+
+	}  // End namespace Private
+
+	class FMeshResizingDataflowNodesModule : public IModuleInterface
+	{
+	public:
+		virtual void StartupModule() override
+		{
+			Private::RegisterDataflowNodes();
+		}
+
+		virtual void ShutdownModule() override
+		{
+		}
+	};
+}  // End namespace UE::MeshResizing
+
+IMPLEMENT_MODULE(UE::MeshResizing::FMeshResizingDataflowNodesModule, MeshResizingDataflowNodes)
diff --git a/Engine/Plugins/Experimental/MeshResizing/Source/MeshResizingNodes/Public/MeshResizing/BaseBodyDataflowNodes.h b/Engine/Plugins/Experimental/MeshResizing/Source/MeshResizingNodes/Public/MeshResizing/BaseBodyDataflowNodes.h
new file mode 100644
--- /dev/null
+++ b/Engine/Plugins/Experimental/MeshResizing/Source/MeshResizingNodes/Public/MeshResizing/BaseBodyDataflowNodes.h
@@ -0,0 +1,60 @@
+// Copyright Epic Games, Inc. All Rights Reserved.
+
+#pragma once 
+
+#include "Dataflow/DataflowEngine.h"
+#include "DynamicMesh/DynamicMesh3.h"
+#include "DynamicMeshProcessor.h"
+
+#include "BaseBodyDataflowNodes.generated.h"
+
+class USkeletalMesh;
+class UDynamicMesh;
+
+/**
+ *
+ * Converts a SkeletalMesh into a DynamicMesh with Imported Vertex information
+ *
+ */
+USTRUCT(meta = (MeshResizing))
+struct FSkeletalMeshToMeshDataflowNode : public FDataflowNode
+{
+	GENERATED_USTRUCT_BODY()
+	DATAFLOW_NODE_DEFINE_INTERNAL(FSkeletalMeshToMeshDataflowNode, "SkeletalMeshToMesh", "Mesh|Utilities", "")
+	DATAFLOW_NODE_RENDER_TYPE("SurfaceRender",FName("FDynamicMesh3"), "Mesh")
+
+public:
+
+	/** SkeletalMesh to convert */
+	UPROPERTY(EditAnywhere, Category = "SkeletalMesh", meta = (DataflowInput));
+	TObjectPtr<const USkeletalMesh> SkeletalMesh;
+
+	/** Output mesh */
+	UPROPERTY(meta = (DataflowOutput))
+	TObjectPtr<UDynamicMesh> Mesh;
+
+	FSkeletalMeshToMeshDataflowNode(const UE::Dataflow::FNodeParameters& InParam, FGuid InGuid = FGuid::NewGuid())
+		: FDataflowNode(InParam, InGuid)
+	{
+		RegisterInputConnection(&SkeletalMesh);
+		RegisterOutputConnection(&Mesh);
+	}
+
+private:
+
+	virtual void Evaluate(UE::Dataflow::FContext& Context, const FDataflowOutput* Out) const override;
+
+	/** Specifies the LOD level to use */
+	UPROPERTY(EditAnywhere, Category = "SkeletalMesh", meta = (DisplayName = "LOD Level"));
+	int32 LODLevel = 0;
+
+	/** Record ImportedVertices (if available) as NonManifold mapping data.*/
+	UPROPERTY(EditAnywhere, Category = "SkeletalMesh")
+	bool bRecordImportedVertices = true;
+};
+
+namespace UE::MeshResizing
+{
+	void RegisterBaseBodyDataflowNodes();
+}
+
