Dataflow Rendering: Add ability to render DynamicMeshes with materials.

Updated the StaticMeshToMesh node to output the material data. Also updated it to accept a Static Mesh input.

Updated new SkeletalMeshToMesh node to also output material data and to record ImportVertex info. (NOTE: this node is part of an experimental plugin. It might move to GeometryCollectionMeshNodes at some point once we stop churning on it.)

diff --git a/Engine/Plugins/Experimental/Dataflow/Source/DataflowEditor/Private/Dataflow/DataflowEngineRendering.cpp b/Engine/Plugins/Experimental/Dataflow/Source/DataflowEditor/Private/Dataflow/DataflowEngineRendering.cpp
--- a/Engine/Plugins/Experimental/Dataflow/Source/DataflowEditor/Private/Dataflow/DataflowEngineRendering.cpp
+++ b/Engine/Plugins/Experimental/Dataflow/Source/DataflowEditor/Private/Dataflow/DataflowEngineRendering.cpp
@@ -12,6 +12,7 @@
 #include "GeometryCollection/Facades/CollectionExplodedVectorFacade.h"
 #include "GeometryCollection/ManagedArrayCollection.h"
 #include "GeometryCollection/GeometryCollectionAlgo.h"
+#include "Materials/MaterialInterface.h"
 #include "UDynamicMesh.h"
 
 namespace UE::Dataflow
@@ -428,11 +429,18 @@ namespace UE::Dataflow
 					{
 						// This will contain the valid triangles only
 						TArray<FIntVector> Tris; Tris.Reserve(DynamicMesh.TriangleCount());
+						TArray<int32> MaterialIDs; MaterialIDs.Reserve(DynamicMesh.TriangleCount());
+						const UE::Geometry::FDynamicMeshMaterialAttribute* const MaterialAttribute = DynamicMesh.Attributes()->GetMaterialID();
 
 						// DynamicMesh.TrianglesItr() returns the valid triangles only
-						for (UE::Geometry::FIndex3i Tri : DynamicMesh.TrianglesItr())
+						for (const int32 TriangleID : DynamicMesh.TriangleIndicesItr())
 						{
+							const UE::Geometry::FIndex3i Tri = DynamicMesh.GetTriangle(TriangleID);
 							Tris.Add(FIntVector(Tri.A, Tri.B, Tri.C));
+							if (MaterialAttribute)
+							{
+								MaterialIDs.Add(MaterialAttribute->GetValue(TriangleID));
+							}
 						}
 
 						// This will contain all the vertices (invalid ones too)
@@ -455,8 +463,70 @@ namespace UE::Dataflow
 							VertexColors[VertexIdx] = FLinearColor(FDataflowEditorModule::SurfaceColor);
 						}
 
-						int32 GeometryIndex = RenderCollection.StartGeometryGroup(State.GetGuid().ToString());
-						RenderCollection.AddSurface(MoveTemp(Vertices), MoveTemp(Tris), MoveTemp(VertexNormals), MoveTemp(VertexColors));
+						// See if we have materials to render.
+						TArray<TObjectPtr<UMaterialInterface>> Materials;
+						if (State.GetRenderOutputs().IsValidIndex(1))
+						{
+							TArray<TObjectPtr<UMaterialInterface>> MaterialsDefault;
+							Materials = State.GetValue<TArray<TObjectPtr<UMaterialInterface>>>(State.GetRenderOutputs()[1], MaterialsDefault);
+						}
+
+						const int32 NumUVLayers = DynamicMesh.Attributes()->NumUVLayers();
+
+						const int32 GeometryIndex = RenderCollection.StartGeometryGroup(State.GetGuid().ToString());
+
+						if (MaterialIDs.Num() != Tris.Num() || Materials.IsEmpty() || NumUVLayers == 0)
+						{
+							// No materials.
+							RenderCollection.AddSurface(MoveTemp(Vertices), MoveTemp(Tris), MoveTemp(VertexNormals), MoveTemp(VertexColors));
+						}
+						else
+						{
+							// GetUVs and materials.
+							// The render facade currently assumes there are no internal uv seams (i.e., each vertex has one UV, instead of wedges). 
+							// This will be the case for meshes that come from SKMs/SMs/Cloth
+							TArray<const UE::Geometry::FDynamicMeshUVOverlay*> UVLayers;
+							UVLayers.SetNumUninitialized(NumUVLayers);
+							for (int32 Layer = 0; Layer < NumUVLayers; ++Layer)
+							{
+								UVLayers[Layer] = DynamicMesh.Attributes()->GetUVLayer(Layer);
+							}
+
+							TArray<FVector2f> DefaultUVs;
+							DefaultUVs.SetNumZeroed(NumUVLayers);
+
+							TArray<TArray<FVector2f>> UVs;
+							UVs.Init(DefaultUVs, Vertices.Num()); // RenderFacade->DynamicMesh code expects all vertices to have full UV sets.
+
+							for (const int32 TriangleIndex : DynamicMesh.TriangleIndicesItr())
+							{
+								const UE::Geometry::FIndex3i Tri = DynamicMesh.GetTriangle(TriangleIndex);
+
+								for (int32 TriangleVertexIndex = 0; TriangleVertexIndex < 3; ++TriangleVertexIndex)
+								{
+									const int32 VertexIndex = Tri[TriangleVertexIndex];
+
+									for (int32 Layer = 0; Layer < NumUVLayers; ++Layer)
+									{
+										if (UVLayers[Layer])
+										{
+											UVs[VertexIndex][Layer] = UVLayers[Layer]->GetElementAtVertex(TriangleIndex, VertexIndex);
+										}
+									}
+								}
+							}
+
+							TArray<FString> MaterialPaths;
+							MaterialPaths.Reserve(Materials.Num());
+							for (const TObjectPtr<UMaterialInterface>& Material : Materials)
+							{
+								MaterialPaths.Add(Material ? Material->GetPathName() : "");
+							}
+
+							RenderCollection.AddSurface(MoveTemp(Vertices), MoveTemp(Tris), MoveTemp(VertexNormals), MoveTemp(VertexColors), MoveTemp(UVs), MoveTemp(MaterialIDs), MoveTemp(MaterialPaths));
+
+						}
+
 						RenderCollection.EndGeometryGroup(GeometryIndex);
 					}
 				}
diff --git a/Engine/Plugins/Experimental/GeometryCollectionPlugin/Source/GeometryCollectionNodes/Private/Dataflow/GeometryCollectionMeshNodes.cpp b/Engine/Plugins/Experimental/GeometryCollectionPlugin/Source/GeometryCollectionNodes/Private/Dataflow/GeometryCollectionMeshNodes.cpp
--- a/Engine/Plugins/Experimental/GeometryCollectionPlugin/Source/GeometryCollectionNodes/Private/Dataflow/GeometryCollectionMeshNodes.cpp
+++ b/Engine/Plugins/Experimental/GeometryCollectionPlugin/Source/GeometryCollectionNodes/Private/Dataflow/GeometryCollectionMeshNodes.cpp
@@ -219,9 +219,9 @@ void FStaticMeshToMeshDataflowNode::Evaluate(UE::Dataflow::FContext& Context, co
 #if WITH_EDITORONLY_DATA
 	if (Out->IsA(&Mesh))
 	{
-		if (StaticMesh)
+		if (const UStaticMesh* const InStaticMesh = GetValue(Context, &StaticMesh))
 		{
-			if (FMeshDescription* MeshDescription = bUseHiRes ? StaticMesh->GetHiResMeshDescription() : StaticMesh->GetMeshDescription(LODLevel))
+			if (const FMeshDescription* const MeshDescription = bUseHiRes ? InStaticMesh->GetHiResMeshDescription() : InStaticMesh->GetMeshDescription(LODLevel))
 			{
 				TObjectPtr<UDynamicMesh> NewMesh = NewObject<UDynamicMesh>();
 				NewMesh->Reset();
@@ -238,6 +238,41 @@ void FStaticMeshToMeshDataflowNode::Evaluate(UE::Dataflow::FContext& Context, co
 		}
 		SetValue(Context, TObjectPtr<UDynamicMesh>(NewObject<UDynamicMesh>()), &Mesh);
 	}
+	else if (Out->IsA(&MaterialArray))
+	{
+		// The dynamic mesh converter will set the MaterialIDs = PolyGroupID by default.
+		// Output materials to match this.
+		TArray<TObjectPtr<UMaterialInterface>> OutMaterials;
+		if (const UStaticMesh* const InStaticMesh = GetValue(Context, &StaticMesh))
+		{
+			const TArray<FStaticMaterial>& StaticMaterials = InStaticMesh->GetStaticMaterials();
+			if (const FMeshDescription* const MeshDescription = bUseHiRes ? InStaticMesh->GetHiResMeshDescription() : InStaticMesh->GetMeshDescription(LODLevel))
+			{
+				if (bUseHiRes)
+				{
+					const FStaticMeshConstAttributes MeshDescriptionAttributes(*MeshDescription);
+					TPolygonGroupAttributesConstRef<FName> MaterialSlotNames = MeshDescriptionAttributes.GetPolygonGroupMaterialSlotNames();
+					OutMaterials.Reserve(MaterialSlotNames.GetNumElements());
+					for (int32 PolyGroupID = 0; PolyGroupID < MaterialSlotNames.GetNumElements(); ++PolyGroupID)
+					{
+						const int32 MaterialIndex = InStaticMesh->GetMaterialIndexFromImportedMaterialSlotName(MaterialSlotNames[PolyGroupID]);
+						OutMaterials.Emplace(StaticMaterials.IsValidIndex(MaterialIndex) ? StaticMaterials[MaterialIndex].MaterialInterface : nullptr);
+					}
+				}
+				else
+				{
+					const FMeshSectionInfoMap& SectionMap = InStaticMesh->GetSectionInfoMap();
+					const int32 LODSectionNum = SectionMap.GetSectionNumber(LODLevel);
+					for (int32 SectionIndex = 0; SectionIndex < LODSectionNum; ++SectionIndex)
+					{
+						const int32 MaterialIndex = SectionMap.IsValidSection(LODLevel, SectionIndex) ? SectionMap.Get(LODLevel, SectionIndex).MaterialIndex : INDEX_NONE;
+						OutMaterials.Emplace(StaticMaterials.IsValidIndex(MaterialIndex) ? StaticMaterials[MaterialIndex].MaterialInterface : nullptr);
+					}
+				}
+			}
+		}
+		SetValue(Context, OutMaterials, &MaterialArray);
+	}
 #endif
 }
 
diff --git a/Engine/Plugins/Experimental/GeometryCollectionPlugin/Source/GeometryCollectionNodes/Public/Dataflow/GeometryCollectionMeshNodes.h b/Engine/Plugins/Experimental/GeometryCollectionPlugin/Source/GeometryCollectionNodes/Public/Dataflow/GeometryCollectionMeshNodes.h
--- a/Engine/Plugins/Experimental/GeometryCollectionPlugin/Source/GeometryCollectionNodes/Public/Dataflow/GeometryCollectionMeshNodes.h
+++ b/Engine/Plugins/Experimental/GeometryCollectionPlugin/Source/GeometryCollectionNodes/Public/Dataflow/GeometryCollectionMeshNodes.h
@@ -20,6 +20,7 @@ namespace UE_DEPRECATED(5.5, "Use UE::Dataflow instead.") Dataflow {}
 class FGeometryCollection;
 class UStaticMesh;
 class UDynamicMesh;
+class UMaterialInterface;
 
 /**
  *
@@ -323,11 +324,11 @@ struct FStaticMeshToMeshDataflowNode : public FDataflowNode
 {
 	GENERATED_USTRUCT_BODY()
 	DATAFLOW_NODE_DEFINE_INTERNAL(FStaticMeshToMeshDataflowNode, "StaticMeshToMesh", "Mesh|Utilities", "")
-	DATAFLOW_NODE_RENDER_TYPE("SurfaceRender",FName("FDynamicMesh3"), "Mesh")
+	DATAFLOW_NODE_RENDER_TYPE("SurfaceRender",FName("FDynamicMesh3"), {"Mesh", "MaterialArray"})
 
 public:
 	/** StaticMesh to convert */
-	UPROPERTY(EditAnywhere, Category = "StaticMesh");
+	UPROPERTY(EditAnywhere, Category = "StaticMesh", meta = (DataflowInput));
 	TObjectPtr<UStaticMesh> StaticMesh;
 
 	/** Output the HiRes representation, if set to true and HiRes doesn't exist it will output empty mesh */
@@ -342,10 +343,16 @@ public:
 	UPROPERTY(meta = (DataflowOutput))
 	TObjectPtr<UDynamicMesh> Mesh;
 
+	/** Output materials */
+	UPROPERTY(meta = (DataflowOutput))
+	TArray<TObjectPtr<UMaterialInterface>> MaterialArray;
+
 	FStaticMeshToMeshDataflowNode(const UE::Dataflow::FNodeParameters& InParam, FGuid InGuid = FGuid::NewGuid())
 		: FDataflowNode(InParam, InGuid)
 	{
+		RegisterInputConnection(&StaticMesh);
 		RegisterOutputConnection(&Mesh);
+		RegisterOutputConnection(&MaterialArray);
 	}
 
 	virtual void Evaluate(UE::Dataflow::FContext& Context, const FDataflowOutput* Out) const override;
diff --git a/Engine/Plugins/Experimental/MeshResizing/Source/MeshResizingNodes/Private/MeshResizing/BaseBodyDataflowNodes.cpp b/Engine/Plugins/Experimental/MeshResizing/Source/MeshResizingNodes/Private/MeshResizing/BaseBodyDataflowNodes.cpp
--- a/Engine/Plugins/Experimental/MeshResizing/Source/MeshResizingNodes/Private/MeshResizing/BaseBodyDataflowNodes.cpp
+++ b/Engine/Plugins/Experimental/MeshResizing/Source/MeshResizingNodes/Private/MeshResizing/BaseBodyDataflowNodes.cpp
@@ -3,8 +3,11 @@
 #include "MeshResizing/BaseBodyDataflowNodes.h"
 #include "Dataflow/DataflowCore.h"
 #include "DynamicMesh/DynamicMesh3.h"
+#include "DynamicMesh/NonManifoldMappingSupport.h"
 #include "Engine/SkeletalMesh.h"
 #include "Rendering/SkeletalMeshRenderData.h"
+#include "Rendering/SkeletalMeshModel.h"
+#include "SkeletalMeshLODModelToDynamicMesh.h"
 #include "SkeletalMeshLODRenderDataToDynamicMesh.h"
 #include "UDynamicMesh.h"
 
@@ -26,27 +29,82 @@ void FSkeletalMeshToMeshDataflowNode::Evaluate(UE::Dataflow::FContext& Context,
 	{
 		if (const USkeletalMesh* const InSkeletalMesh = GetValue(Context, &SkeletalMesh))
 		{
-			if (const FSkeletalMeshRenderData* const RenderData = InSkeletalMesh->GetResourceForRendering())
+			if (bRecordImportedVertices)
 			{
-				if (RenderData->LODRenderData.IsValidIndex(LODLevel))
+#if WITH_EDITORONLY_DATA
+				if (const FSkeletalMeshModel* const MeshModel = InSkeletalMesh->GetImportedModel())
 				{
-					const FSkeletalMeshLODRenderData* const SkeletalMeshLODRenderData = &RenderData->LODRenderData[LODLevel];
+					if (MeshModel->LODModels.IsValidIndex(LODLevel))
+					{
+						const FSkeletalMeshLODModel* const SkeletalMeshLODModel = &MeshModel->LODModels[LODLevel];
 
-					TObjectPtr<UDynamicMesh> NewMesh = NewObject<UDynamicMesh>();
-					NewMesh->Reset();
+						TObjectPtr<UDynamicMesh> NewMesh = NewObject<UDynamicMesh>();
+						NewMesh->Reset();
 
-					UE::Geometry::FDynamicMesh3& DynMesh = NewMesh->GetMeshRef();
+						UE::Geometry::FDynamicMesh3& DynMesh = NewMesh->GetMeshRef();
 
-					UE::Geometry::FSkeletalMeshLODRenderDataToDynamicMesh::ConversionOptions  ConversionOptions;
-					//ConversionOptions.bWantTangents = AssetOptions.bRequestTangents;
+						FSkeletalMeshLODModelToDynamicMesh Converter;
+						Converter.Convert(SkeletalMeshLODModel, DynMesh);
 
-					UE::Geometry::FSkeletalMeshLODRenderDataToDynamicMesh::Convert(SkeletalMeshLODRenderData, InSkeletalMesh->GetRefSkeleton(), ConversionOptions, DynMesh);
+						// Set ImportedVertices
+						TArray<int32> VertexToImportedVertexIDMap;
+						VertexToImportedVertexIDMap.SetNumUninitialized(DynMesh.VertexCount());
+						check(DynMesh.VertexCount() == Converter.VertIDMap.Num());
 
-					SetValue(Context, NewMesh, &Mesh);
-					return;
+						const TArray<int32>& LODModelToImportVertexMap = SkeletalMeshLODModel->MeshToImportVertexMap;
+						
+						for (int32 DynMeshVertexIdx = 0; DynMeshVertexIdx < DynMesh.VertexCount(); ++DynMeshVertexIdx)
+						{
+							const int32 LODModelVertexIdx = Converter.VertIDMap[DynMeshVertexIdx];
+							const int32 ImportVertexIdx = LODModelToImportVertexMap[LODModelVertexIdx];
+							VertexToImportedVertexIDMap[DynMeshVertexIdx] = ImportVertexIdx;
+						}
+						UE::Geometry::FNonManifoldMappingSupport::AttachNonManifoldVertexMappingData(VertexToImportedVertexIDMap, DynMesh);
+
+						SetValue(Context, NewMesh, &Mesh);
+						return;
+					}
+				}
+#else
+				// Add error message about how this doesn't work? TBH not sure what to do here because if you started with the imported vertices and now they're gone, everything down stream is going to be broken
+#endif
+			}
+			else
+			{
+				if (const FSkeletalMeshRenderData* const RenderData = InSkeletalMesh->GetResourceForRendering())
+				{
+					if (RenderData->LODRenderData.IsValidIndex(LODLevel))
+					{
+						const FSkeletalMeshLODRenderData* const SkeletalMeshLODRenderData = &RenderData->LODRenderData[LODLevel];
+
+						TObjectPtr<UDynamicMesh> NewMesh = NewObject<UDynamicMesh>();
+						NewMesh->Reset();
+
+						UE::Geometry::FDynamicMesh3& DynMesh = NewMesh->GetMeshRef();
+
+						UE::Geometry::FSkeletalMeshLODRenderDataToDynamicMesh::ConversionOptions  ConversionOptions;
+						UE::Geometry::FSkeletalMeshLODRenderDataToDynamicMesh::Convert(SkeletalMeshLODRenderData, InSkeletalMesh->GetRefSkeleton(), ConversionOptions, DynMesh);
+
+						SetValue(Context, NewMesh, &Mesh);
+						return;
+					}
 				}
 			}
 		}
 		SetValue(Context, TObjectPtr<UDynamicMesh>(NewObject<UDynamicMesh>()), &Mesh);
 	}
+	else if (Out->IsA(&MaterialArray))
+	{
+		TArray<TObjectPtr<UMaterialInterface>> OutMaterials;
+		if (const USkeletalMesh* const InSkeletalMesh = GetValue(Context, &SkeletalMesh))
+		{
+			const TArray<FSkeletalMaterial>& InMaterials = InSkeletalMesh->GetMaterials();
+			OutMaterials.Reserve(InMaterials.Num());
+			for (const FSkeletalMaterial& Material : InMaterials)
+			{
+				OutMaterials.Emplace(Material.MaterialInterface);
+			}
+		}
+		SetValue(Context, OutMaterials, &MaterialArray);
+	}
 }
diff --git a/Engine/Plugins/Experimental/MeshResizing/Source/MeshResizingNodes/Public/MeshResizing/BaseBodyDataflowNodes.h b/Engine/Plugins/Experimental/MeshResizing/Source/MeshResizingNodes/Public/MeshResizing/BaseBodyDataflowNodes.h
--- a/Engine/Plugins/Experimental/MeshResizing/Source/MeshResizingNodes/Public/MeshResizing/BaseBodyDataflowNodes.h
+++ b/Engine/Plugins/Experimental/MeshResizing/Source/MeshResizingNodes/Public/MeshResizing/BaseBodyDataflowNodes.h
@@ -10,6 +10,7 @@
 
 class USkeletalMesh;
 class UDynamicMesh;
+class UMaterialInterface;
 
 /**
  *
@@ -21,7 +22,7 @@ struct FSkeletalMeshToMeshDataflowNode : public FDataflowNode
 {
 	GENERATED_USTRUCT_BODY()
 	DATAFLOW_NODE_DEFINE_INTERNAL(FSkeletalMeshToMeshDataflowNode, "SkeletalMeshToMesh", "Mesh|Utilities", "")
-	DATAFLOW_NODE_RENDER_TYPE("SurfaceRender",FName("FDynamicMesh3"), "Mesh")
+	DATAFLOW_NODE_RENDER_TYPE("SurfaceRender",FName("FDynamicMesh3"), {"Mesh", "MaterialArray"})
 
 public:
 
@@ -33,11 +34,16 @@ public:
 	UPROPERTY(meta = (DataflowOutput))
 	TObjectPtr<UDynamicMesh> Mesh;
 
+	/** Output materials */
+	UPROPERTY(meta = (DataflowOutput))
+	TArray<TObjectPtr<UMaterialInterface>> MaterialArray;
+
 	FSkeletalMeshToMeshDataflowNode(const UE::Dataflow::FNodeParameters& InParam, FGuid InGuid = FGuid::NewGuid())
 		: FDataflowNode(InParam, InGuid)
 	{
 		RegisterInputConnection(&SkeletalMesh);
 		RegisterOutputConnection(&Mesh);
+		RegisterOutputConnection(&MaterialArray);
 	}
 
 private:
@@ -48,7 +54,7 @@ private:
 	UPROPERTY(EditAnywhere, Category = "SkeletalMesh", meta = (DisplayName = "LOD Level"));
 	int32 LODLevel = 0;
 
-	/** Record ImportedVertices (if available) as NonManifold mapping data.*/
+	/** Record ImportedVertices (if available) as NonManifold mapping data. This requires Editor-Only data.*/
 	UPROPERTY(EditAnywhere, Category = "SkeletalMesh")
 	bool bRecordImportedVertices = true;
 };
