Adding a MeshWarp node that changes the shape of a mesh based on a source (bind) mesh and target (deformed) mesh

diff --git a/Engine/Plugins/Experimental/Dataflow/Source/DataflowAssetTools/Private/Dataflow/CollectionRenderingPatternUtility.cpp b/Engine/Plugins/Experimental/Dataflow/Source/DataflowAssetTools/Private/Dataflow/CollectionRenderingPatternUtility.cpp
--- a/Engine/Plugins/Experimental/Dataflow/Source/DataflowAssetTools/Private/Dataflow/CollectionRenderingPatternUtility.cpp
+++ b/Engine/Plugins/Experimental/Dataflow/Source/DataflowAssetTools/Private/Dataflow/CollectionRenderingPatternUtility.cpp
@@ -102,12 +102,13 @@ namespace UE::Dataflow
 					MaterialIDAttrib->SetValue(TriangleID, FacadeMaterialID[TriangleID] - StartMaterials);
 				}
 
-				const int32 NumUVLayers = UVs[0].Num();
+				const int32 NumUVLayers = Remapping.Num() > 0 ? UVs[Remapping[0]].Num() : 0;
 
 				if (NumUVLayers > 0)
 				{
-					for (const TArray<FVector2f>& VertexUVs : UVs)
+					for (int32 VertexIndex = StartVertices; VertexIndex < LastVertexIndex; ++VertexIndex)
 					{
+						const TArray<FVector2f>& VertexUVs = UVs[VertexIndex];
 						check(NumUVLayers == VertexUVs.Num());
 					}
 
diff --git a/Engine/Plugins/Experimental/MeshResizing/Source/MeshResizingNodes/Private/MeshResizing/MeshResizingDataflowNodesModule.cpp b/Engine/Plugins/Experimental/MeshResizing/Source/MeshResizingNodes/Private/MeshResizing/MeshResizingDataflowNodesModule.cpp
--- a/Engine/Plugins/Experimental/MeshResizing/Source/MeshResizingNodes/Private/MeshResizing/MeshResizingDataflowNodesModule.cpp
+++ b/Engine/Plugins/Experimental/MeshResizing/Source/MeshResizingNodes/Private/MeshResizing/MeshResizingDataflowNodesModule.cpp
@@ -4,6 +4,7 @@
 #include "Dataflow/DataflowNodeFactory.h"
 #include "Modules/ModuleManager.h"
 #include "MeshResizing/BaseBodyDataflowNodes.h"
+#include "MeshResizing/MeshWarpNode.h"
 
 namespace UE::MeshResizing
 {
@@ -21,6 +22,8 @@ namespace UE::MeshResizing
 			DATAFLOW_NODE_REGISTER_CREATION_FACTORY_NODE_COLORS_BY_CATEGORY("MeshResizing", FColorScheme::NodeHeader, FColorScheme::NodeBody);
 
 			RegisterBaseBodyDataflowNodes();
+
+			DATAFLOW_NODE_REGISTER_CREATION_FACTORY(FMeshWarpNode);
 		}
 
 	}  // End namespace Private
diff --git a/Engine/Plugins/Experimental/MeshResizing/Source/MeshResizingNodes/Private/MeshResizing/MeshWarpNode.cpp b/Engine/Plugins/Experimental/MeshResizing/Source/MeshResizingNodes/Private/MeshResizing/MeshWarpNode.cpp
new file mode 100644
--- /dev/null
+++ b/Engine/Plugins/Experimental/MeshResizing/Source/MeshResizingNodes/Private/MeshResizing/MeshWarpNode.cpp
@@ -0,0 +1,163 @@
+// Copyright Epic Games, Inc. All Rights Reserved.
+#include "MeshResizing/MeshWarpNode.h"
+#include "DynamicMesh/DynamicMesh3.h"
+#include "DynamicMesh/DynamicMeshAABBTree3.h"
+#include "UDynamicMesh.h"
+#include "DynamicMeshEditor.h"
+#include "Materials/MaterialInterface.h"
+
+#include UE_INLINE_GENERATED_CPP_BY_NAME(MeshWarpNode)
+
+namespace  UE::Chaos::MeshResizing::Private
+{
+	void BlendMesh(
+		FDynamicMesh3* OutMesh,
+		const UDynamicMesh* SourceUMesh,
+		const UDynamicMesh* TargetUMesh,
+		float Alpha)
+	{
+		using namespace UE::Geometry;
+
+		if (OutMesh == nullptr)
+		{
+			UE_LOG(LogTemp, Warning, TEXT("BlendMesh: OutMesh is Null"));
+			return;
+		}
+
+		const FDynamicMesh3* const SourceMesh = SourceUMesh->GetMeshPtr();
+		OutMesh->Copy(*SourceMesh);
+
+		const FDynamicMesh3* const TargetMesh = TargetUMesh->GetMeshPtr();
+
+		ensure(SourceMesh->VertexCount() == TargetMesh->VertexCount());
+		ensure(SourceMesh->VertexCount() == OutMesh->VertexCount());
+
+		for (const int32 VertexIndex : TargetMesh->VertexIndicesItr())
+		{
+			const FVector3d SourceVert = SourceMesh->GetVertex(VertexIndex);
+			const FVector3d TargetVert = TargetMesh->GetVertex(VertexIndex);
+			const FVector3d FinalPoint = (1.0 - Alpha) * SourceVert + Alpha * TargetVert;
+			OutMesh->SetVertex(VertexIndex, FinalPoint);
+		}
+	}
+
+
+	void ApplyRefitToMesh(
+		FDynamicMesh3* ResizedMesh,
+		const UDynamicMesh* SourceUMesh,
+		const FDynamicMesh3& TargetMesh,
+		float Alpha)
+	{
+		using namespace UE::Geometry;
+
+		if (ResizedMesh == nullptr)
+		{
+			UE_LOG(LogTemp, Warning, TEXT("ApplyRefitToMesh: ResizedMesh is Null"));
+			return;
+		}
+
+		const FDynamicMesh3* const SourceMesh = SourceUMesh->GetMeshPtr();
+		FDynamicMeshAABBTree3 SourceTree(SourceMesh);
+
+		struct FEmbeddingInfo
+		{
+			int ClosestTriangle;
+			FVector4d Weights;
+		};
+
+		TArray<FEmbeddingInfo> EmbeddingInfo;
+		EmbeddingInfo.SetNum(ResizedMesh->VertexCount());
+
+		for (const int32 VertexIndex : ResizedMesh->VertexIndicesItr())
+		{
+			EmbeddingInfo[VertexIndex].ClosestTriangle = INDEX_NONE;
+
+			const FVector3d VertexLocation = ResizedMesh->GetVertex(VertexIndex);
+
+			double NearestDistSqr;
+			const int32 NearTriID = SourceTree.FindNearestTriangle(VertexLocation, NearestDistSqr);
+			if (NearTriID >= 0)
+			{
+				const FVector3d Normal = SourceMesh->GetTriNormal(NearTriID);
+				FVector3d TriPointA, TriPointB, TriPointC;
+				SourceMesh->GetTriVertices(NearTriID, TriPointA, TriPointB, TriPointC);
+				const FPlane3d TriPlane(Normal, TriPointA);
+				const double Dist = TriPlane.DistanceTo(VertexLocation);
+				const FVector3d PointOnPlane = VertexLocation - Dist * Normal;
+				const FVector Bary = VectorUtil::BarycentricCoords(PointOnPlane, TriPointA, TriPointB, TriPointC);
+
+				EmbeddingInfo[VertexIndex].ClosestTriangle = NearTriID;
+				EmbeddingInfo[VertexIndex].Weights = { Bary, Dist };
+			}
+		}
+
+		for (const int32 VertexIndex : ResizedMesh->VertexIndicesItr())
+		{
+			if (EmbeddingInfo[VertexIndex].ClosestTriangle != INDEX_NONE)
+			{
+				const FVector3d Normal = SourceMesh->GetTriNormal(EmbeddingInfo[VertexIndex].ClosestTriangle);
+				FVector3d TriPointA, TriPointB, TriPointC;
+				TargetMesh.GetTriVertices(EmbeddingInfo[VertexIndex].ClosestTriangle, TriPointA, TriPointB, TriPointC);
+
+				const FVector4d& Weights = EmbeddingInfo[VertexIndex].Weights;
+				const FVector PointOnPlane = Weights[0] * TriPointA + Weights[1] * TriPointB + Weights[2] * TriPointC;
+				const FVector3d CurrentVertexPosition = ResizedMesh->GetVertex(VertexIndex);
+				const FVector FinalPoint = (1.0 - Alpha) * CurrentVertexPosition + Alpha * (PointOnPlane + Weights[3] * Normal);
+
+				ResizedMesh->SetVertex(VertexIndex, FinalPoint);
+			}
+		}
+	}
+}
+
+
+FMeshWarpNode::FMeshWarpNode(const UE::Dataflow::FNodeParameters& InParam, FGuid InGuid) :
+	FDataflowNode(InParam, InGuid)
+{
+	RegisterInputConnection(&MeshToResize);
+	RegisterInputConnection(&MeshToResizeMaterialArray);
+	RegisterInputConnection(&SourceMesh);
+	RegisterInputConnection(&TargetMesh);
+
+	RegisterOutputConnection(&BlendedTargetMesh);
+	RegisterOutputConnection(&ResizedMesh);
+	RegisterOutputConnection(&MaterialArray);
+}
+
+void FMeshWarpNode::Evaluate(UE::Dataflow::FContext& Context, const FDataflowOutput* Out) const
+{
+	using namespace UE::Geometry;
+
+	TObjectPtr<UDynamicMesh> OutResizedMesh = NewObject<UDynamicMesh>();
+	TObjectPtr<UDynamicMesh> OutBlendedTargetMesh = NewObject<UDynamicMesh>();
+	TArray<TObjectPtr<UMaterialInterface>> OutMaterials;
+
+	if (TObjectPtr<UDynamicMesh> InMeshToResize = GetValue(Context, &MeshToResize))
+	{
+		if (TObjectPtr<UDynamicMesh> InSourceMesh = GetValue(Context, &SourceMesh))
+		{
+			if (TObjectPtr<UDynamicMesh> InTargetMesh = GetValue(Context, &TargetMesh))
+			{
+				FDynamicMesh3 BlendedTargetFMesh;
+				UE::Chaos::MeshResizing::Private::BlendMesh(&BlendedTargetFMesh, InSourceMesh, InTargetMesh, Alpha);
+
+				FDynamicMesh3 ResizeMesh;
+				ResizeMesh.Copy(InMeshToResize->GetMeshRef());
+				UE::Chaos::MeshResizing::Private::ApplyRefitToMesh(&ResizeMesh, InSourceMesh, BlendedTargetFMesh, 1.0f);
+
+				OutBlendedTargetMesh->SetMesh(MoveTemp(BlendedTargetFMesh));
+				OutResizedMesh->SetMesh(MoveTemp(ResizeMesh));
+
+				if (IsConnected(&MeshToResizeMaterialArray))
+				{
+					OutMaterials = GetValue(Context, &MeshToResizeMaterialArray);
+				}
+			}
+		}
+	}
+
+	SetValue(Context, MoveTemp(OutResizedMesh), &ResizedMesh);
+	SetValue(Context, MoveTemp(OutBlendedTargetMesh), &BlendedTargetMesh);
+	SetValue(Context, MoveTemp(OutMaterials), &MaterialArray);
+}
+
diff --git a/Engine/Plugins/Experimental/MeshResizing/Source/MeshResizingNodes/Private/MeshResizing/MeshWarpNode.h b/Engine/Plugins/Experimental/MeshResizing/Source/MeshResizingNodes/Private/MeshResizing/MeshWarpNode.h
new file mode 100644
--- /dev/null
+++ b/Engine/Plugins/Experimental/MeshResizing/Source/MeshResizingNodes/Private/MeshResizing/MeshWarpNode.h
@@ -0,0 +1,58 @@
+// Copyright Epic Games, Inc. All Rights Reserved.
+
+#pragma once
+
+#include "Dataflow/DataflowNode.h"
+#include "GeometryCollection/ManagedArrayCollection.h"
+
+#include "MeshWarpNode.generated.h"
+
+class UDynamicMesh;
+class UMaterialInterface;
+
+USTRUCT(Meta = (MeshResizing, Experimental))
+struct FMeshWarpNode : public FDataflowNode
+{
+	GENERATED_USTRUCT_BODY()
+	DATAFLOW_NODE_DEFINE_INTERNAL(FMeshWarpNode, "MeshWarp", "MeshResizing", "Mesh Warp")
+	DATAFLOW_NODE_RENDER_TYPE_START()
+	DATAFLOW_NODE_RENDER_TYPE_ADD("SurfaceRender", FName("FDynamicMesh3"), "BlendedTargetMesh")
+	DATAFLOW_NODE_RENDER_TYPE_ADD("SurfaceRender", FName("FDynamicMesh3"), { "ResizedMesh", "MaterialArray" })
+	DATAFLOW_NODE_RENDER_TYPE_END()
+
+public:
+
+	FMeshWarpNode(const UE::Dataflow::FNodeParameters& InParam, FGuid InGuid = FGuid::NewGuid());
+
+private:
+
+	UPROPERTY(meta = (DataflowInput))
+	TObjectPtr<UDynamicMesh> MeshToResize;
+
+	UPROPERTY(meta = (DataflowInput))
+	TArray<TObjectPtr<UMaterialInterface>> MeshToResizeMaterialArray;
+
+	UPROPERTY(meta = (DataflowInput))
+	TObjectPtr<UDynamicMesh> SourceMesh;
+
+	UPROPERTY(meta = (DataflowInput))
+	TObjectPtr<UDynamicMesh> TargetMesh;
+
+
+	UPROPERTY(meta = (DataflowOutput))
+	TObjectPtr<UDynamicMesh> BlendedTargetMesh;
+
+	UPROPERTY(meta = (DataflowOutput))
+	TObjectPtr<UDynamicMesh> ResizedMesh;
+
+	UPROPERTY(meta = (DataflowOutput))
+	TArray<TObjectPtr<UMaterialInterface>> MaterialArray;
+
+
+	UPROPERTY(EditAnywhere, Category = "Warp", meta = (Min = 0.0f, UIMin = 0.0f, Max = 1.0f, UIMax = 1.0f))
+	float Alpha;
+
+	//~ Begin FDataflowNode interface
+	virtual void Evaluate(UE::Dataflow::FContext& Context, const FDataflowOutput* Out) const override;
+	//~ End FDataflowNode interface
+};
