Add some nodes to the MeshResizing plugin that can generate blendable DynamicMeshes from SKMs that had matching source topology.

diff --git a/Engine/Plugins/Experimental/MeshResizing/MeshResizing.uplugin b/Engine/Plugins/Experimental/MeshResizing/MeshResizing.uplugin
--- a/Engine/Plugins/Experimental/MeshResizing/MeshResizing.uplugin
+++ b/Engine/Plugins/Experimental/MeshResizing/MeshResizing.uplugin
@@ -15,6 +15,11 @@
 	"IsExperimentalVersion": true,
     "Installed" : false,
 	"Modules": [
+		{
+			"Name": "MeshResizingCore",
+			"Type": "Runtime",
+			"LoadingPhase": "Default"
+		},
 		{
 			"Name": "MeshResizingDataflowNodes",
 			"Type": "Runtime",
diff --git a/Engine/Plugins/Experimental/MeshResizing/Source/MeshResizingCore/MeshResizingCore.Build.cs b/Engine/Plugins/Experimental/MeshResizing/Source/MeshResizingCore/MeshResizingCore.Build.cs
--- a/Engine/Plugins/Experimental/MeshResizing/Source/MeshResizingCore/MeshResizingCore.Build.cs
+++ b/Engine/Plugins/Experimental/MeshResizing/Source/MeshResizingCore/MeshResizingCore.Build.cs
@@ -11,6 +11,8 @@ namespace UnrealBuildTool.Rules
 			PrivateDependencyModuleNames.AddRange(
 				new string[]
 				{
+					"Core",
+					"GeometryCore",
 				}
 			);
 		}
diff --git a/Engine/Plugins/Experimental/MeshResizing/Source/MeshResizingCore/Private/MeshResizing/BaseBodyTools.cpp b/Engine/Plugins/Experimental/MeshResizing/Source/MeshResizingCore/Private/MeshResizing/BaseBodyTools.cpp
new file mode 100644
--- /dev/null
+++ b/Engine/Plugins/Experimental/MeshResizing/Source/MeshResizingCore/Private/MeshResizing/BaseBodyTools.cpp
@@ -0,0 +1,241 @@
+// Copyright Epic Games, Inc. All Rights Reserved.
+
+#include "MeshResizing/BaseBodyTools.h"
+#include "DynamicMesh/DynamicMesh3.h"
+#include "DynamicMesh/DynamicMeshAttributeSet.h"
+#include "DynamicMesh/DynamicVertexAttribute.h"
+
+namespace UE::MeshResizing
+{
+	namespace Private
+	{
+		class FVertexMappingAttribute final :
+			public UE::Geometry::TDynamicVertexAttribute<int32,1,UE::Geometry::FDynamicMesh3>
+		{
+			typedef UE::Geometry::TDynamicVertexAttribute<int32, 1, UE::Geometry::FDynamicMesh3> Base;
+		public:
+			FVertexMappingAttribute() = default;
+
+			FVertexMappingAttribute(UE::Geometry::FDynamicMesh3* InParent) 
+				: Base(InParent, false)
+			{
+			}
+
+			virtual ~FVertexMappingAttribute() override = default;
+
+			virtual int32 GetDefaultAttributeValue() override
+			{
+				return INDEX_NONE;
+			}
+
+			void InitializeFromArray(TConstArrayView<int32> VertexIDMap)
+			{
+				check(VertexIDMap.Num() == Parent->MaxVertexID());
+				AttribValues.Resize(Parent->MaxVertexID());
+
+				for (int32 VertexID : Parent->VertexIndicesItr())
+				{
+					SetMappedValue(VertexID, VertexIDMap[VertexID]);
+				}
+			}
+
+			void SetMappedValue(int32 VertexID, int32 Value)
+			{
+				AttribValues[VertexID] = Value;
+			}
+
+			int32 GetMappedValue(int32 VertexID) const
+			{
+				return AttribValues[VertexID];
+			}
+
+		private:
+			virtual void SetAttributeFromLerp(int SetAttribute, int AttributeA, int AttributeB, double Alpha) override
+			{
+				// Do not interpolate
+				AttribValues[SetAttribute] = INDEX_NONE;
+			}
+
+			virtual void SetAttributeFromBary(int SetAttribute, int AttributeA, int AttributeB, int AttributeC, const FVector3d& BaryCoords) override
+			{
+				// Do not interpolate
+				AttribValues[SetAttribute] = INDEX_NONE;
+			}
+		};
+
+		static const FVertexMappingAttribute* GetVertexMappingAttribute(const UE::Geometry::FDynamicMesh3& Mesh, const FName& AttrName)
+		{
+			using namespace UE::Geometry;
+			const FDynamicMeshAttributeSet* const Attributes = Mesh.Attributes();
+			if (Attributes && Attributes->HasAttachedAttribute(AttrName))
+			{
+				return static_cast<const FVertexMappingAttribute*>(Attributes->GetAttachedAttribute(AttrName));
+			}
+			return nullptr;
+		}
+
+		static void TransferNormalOverlayValues(const UE::Geometry::FDynamicMeshNormalOverlay& TargetOverlay, UE::Geometry::FDynamicMeshNormalOverlay& ProxyOverlay, const FVertexMappingAttribute& SourceAttr, const TMap<int32, int32>& MappedTargetToMeshTarget)
+		{
+			TArray<int32> TargetVertexElements; // Reused array
+			for (int32 ProxyElementID : ProxyOverlay.ElementIndicesItr())
+			{
+				const int32 ProxyVID = ProxyOverlay.GetParentVertex(ProxyElementID);
+				const int32 MappedVID = SourceAttr.GetMappedValue(ProxyVID);
+				if (const int32* const TargetVID = MappedTargetToMeshTarget.Find(MappedVID))
+				{
+					TargetOverlay.GetVertexElements(*TargetVID, TargetVertexElements);
+					if (ensure(TargetVertexElements.Num() > 0))
+					{
+						ensure(TargetVertexElements.Num() == 1);
+						ProxyOverlay.SetElement(ProxyElementID, TargetOverlay.GetElement(TargetVertexElements[0]));
+					}
+				}
+			}
+		}
+
+		static void InterpolateNormalOverlayValues(const UE::Geometry::FDynamicMeshNormalOverlay& SourceOverlay, const UE::Geometry::FDynamicMeshNormalOverlay& TargetOverlay, float BlendAlpha, UE::Geometry::FDynamicMeshNormalOverlay& ProxyOverlay)
+		{
+			check(SourceOverlay.ElementCount() == TargetOverlay.ElementCount());
+			check(SourceOverlay.ElementCount() == ProxyOverlay.ElementCount());
+			const float OneMinusAlpha = 1.f - BlendAlpha;
+			for (int32 ProxyElementID : ProxyOverlay.ElementIndicesItr())
+			{
+				ProxyOverlay.SetElement(ProxyElementID, OneMinusAlpha * SourceOverlay.GetElement(ProxyElementID) + BlendAlpha * TargetOverlay.GetElement(ProxyElementID));
+			}
+		}
+	}
+
+	const FName FBaseBodyTools::ImportedVertexVIDsAttrName = FName(TEXT("ImportedVertexVIDsAttr"));
+	const FName FBaseBodyTools::RawPointIndicesVIDsAttrName = FName(TEXT("RawPointIndicesVIDsAttr"));
+
+	bool FBaseBodyTools::AttachVertexMappingData(const FName& AttrName, const TArray<int32>& Data, UE::Geometry::FDynamicMesh3& Mesh)
+	{
+		// Copied from FNonManifoldMappingSupport
+		using namespace UE::Geometry;
+		const int32 MaxVID = Mesh.MaxVertexID();
+		FDynamicMeshAttributeSet* const  Attributes = Mesh.Attributes();
+		if (Data.Num() < MaxVID || !Attributes)
+		{
+			return false;
+		}
+
+		// Replace existing buffer.
+		if (Attributes->HasAttachedAttribute(AttrName))
+		{
+			Attributes->RemoveAttribute(AttrName);
+		}
+		Private::FVertexMappingAttribute* SrcMeshVIDAttr = new Private::FVertexMappingAttribute(&Mesh);
+		SrcMeshVIDAttr->SetName(AttrName);
+		Mesh.Attributes()->AttachAttribute(AttrName, SrcMeshVIDAttr);
+		SrcMeshVIDAttr->InitializeFromArray(Data);
+		return true;
+	}
+
+	bool FBaseBodyTools::GenerateResizableProxyFromVertexMappingData(const UE::Geometry::FDynamicMesh3& SourceMesh, const FName& SourceMappingName, const UE::Geometry::FDynamicMesh3& TargetMesh, const FName& TargetMappingName, UE::Geometry::FDynamicMesh3& ProxyMesh)
+	{
+		using namespace UE::Geometry;
+		const Private::FVertexMappingAttribute* const SourceAttr = Private::GetVertexMappingAttribute(SourceMesh, SourceMappingName);
+		const Private::FVertexMappingAttribute* const TargetAttr = Private::GetVertexMappingAttribute(TargetMesh, TargetMappingName);
+		if (!SourceAttr || !TargetAttr)
+		{
+			return false;
+		}
+
+		TMap<int32, int32> MappedTargetToMeshTarget;
+
+		for (int32 VertexID : TargetMesh.VertexIndicesItr())
+		{
+			const int32 TargetValue = TargetAttr->GetMappedValue(VertexID);
+			// Multiple mesh vertices may have the same mapped value. Unless the mesh has been modified, they should all have the same vertex and wedge data since these came from SKMs
+			// Just track last we iterate to.
+			MappedTargetToMeshTarget.Add(TargetValue, VertexID); 
+		}
+
+		// Start with source mesh topology
+		ProxyMesh = SourceMesh;
+
+		// Update vertex data to match target mesh values.
+		for (int32 SourceVertexID : SourceMesh.VertexIndicesItr())
+		{
+			const int32 MappedValue = SourceAttr->GetMappedValue(SourceVertexID);
+
+			if (const int32* const TargetVertexID = MappedTargetToMeshTarget.Find(MappedValue))
+			{
+				ProxyMesh.SetVertex(SourceVertexID, TargetMesh.GetVertexRef(*TargetVertexID));
+			}
+		}
+
+		// Currently assuming materials, uvs, colors match between source and target.
+		
+		const FDynamicMeshNormalOverlay* const TargetNormalOverlay = TargetMesh.Attributes()->PrimaryNormals();
+		const FDynamicMeshNormalOverlay* const TargetTangentOverlay = TargetMesh.Attributes()->PrimaryTangents();
+		const FDynamicMeshNormalOverlay* const TargetBiTangentOverlay = TargetMesh.Attributes()->PrimaryBiTangents();
+
+		FDynamicMeshNormalOverlay* const ProxyNormalOverlay = ProxyMesh.Attributes()->PrimaryNormals();
+		FDynamicMeshNormalOverlay* const ProxyTangentOverlay = ProxyMesh.Attributes()->PrimaryTangents();
+		FDynamicMeshNormalOverlay* const ProxyBiTangentOverlay = ProxyMesh.Attributes()->PrimaryBiTangents();
+
+		if (TargetNormalOverlay && ProxyNormalOverlay)
+		{
+			Private::TransferNormalOverlayValues(*TargetNormalOverlay, *ProxyNormalOverlay, *SourceAttr, MappedTargetToMeshTarget);
+		}
+		if (TargetTangentOverlay && ProxyTangentOverlay)
+		{
+			Private::TransferNormalOverlayValues(*TargetTangentOverlay, *ProxyTangentOverlay, *SourceAttr, MappedTargetToMeshTarget);
+		}
+		if (TargetBiTangentOverlay && ProxyBiTangentOverlay)
+		{
+			Private::TransferNormalOverlayValues(*TargetBiTangentOverlay, *ProxyBiTangentOverlay, *SourceAttr, MappedTargetToMeshTarget);
+		}
+
+		return true;
+	}
+
+	bool FBaseBodyTools::InterpolateResizableProxy(const UE::Geometry::FDynamicMesh3& SourceMesh, const UE::Geometry::FDynamicMesh3& TargetMesh, float BlendAlpha, UE::Geometry::FDynamicMesh3& ProxyMesh)
+	{
+		using namespace UE::Geometry;
+		if (SourceMesh.VertexCount() != TargetMesh.VertexCount())
+		{
+			return false;
+		}
+
+		// in case there are attributes we're not interpolating, choose the closest for those?
+		ProxyMesh = BlendAlpha < 0.5f ? SourceMesh : TargetMesh;
+
+		const float OneMinusAlpha = 1.f - BlendAlpha;
+		for (int32 VertexID : ProxyMesh.VertexIndicesItr())
+		{
+			ProxyMesh.SetVertex(VertexID, OneMinusAlpha * SourceMesh.GetVertexRef(VertexID) + BlendAlpha * TargetMesh.GetVertexRef(VertexID));
+		}
+		// Currently assuming materials, uvs, colors match between source and target.
+
+		const FDynamicMeshNormalOverlay* const SourceNormalOverlay = SourceMesh.Attributes()->PrimaryNormals();
+		const FDynamicMeshNormalOverlay* const SourceTangentOverlay = SourceMesh.Attributes()->PrimaryTangents();
+		const FDynamicMeshNormalOverlay* const SourceBiTangentOverlay = SourceMesh.Attributes()->PrimaryBiTangents();
+
+		const FDynamicMeshNormalOverlay* const TargetNormalOverlay = TargetMesh.Attributes()->PrimaryNormals();
+		const FDynamicMeshNormalOverlay* const TargetTangentOverlay = TargetMesh.Attributes()->PrimaryTangents();
+		const FDynamicMeshNormalOverlay* const TargetBiTangentOverlay = TargetMesh.Attributes()->PrimaryBiTangents();
+
+		FDynamicMeshNormalOverlay* const ProxyNormalOverlay = ProxyMesh.Attributes()->PrimaryNormals();
+		FDynamicMeshNormalOverlay* const ProxyTangentOverlay = ProxyMesh.Attributes()->PrimaryTangents();
+		FDynamicMeshNormalOverlay* const ProxyBiTangentOverlay = ProxyMesh.Attributes()->PrimaryBiTangents();
+
+		if (SourceNormalOverlay && TargetNormalOverlay && SourceNormalOverlay->ElementCount() == TargetNormalOverlay->ElementCount())
+		{
+			check(ProxyNormalOverlay);
+			Private::InterpolateNormalOverlayValues(*SourceNormalOverlay, *TargetNormalOverlay, BlendAlpha, *ProxyNormalOverlay);
+		}
+		if (SourceTangentOverlay && TargetTangentOverlay && SourceTangentOverlay->ElementCount() == TargetTangentOverlay->ElementCount())
+		{
+			check(ProxyTangentOverlay);
+			Private::InterpolateNormalOverlayValues(*SourceTangentOverlay, *TargetTangentOverlay, BlendAlpha, *ProxyTangentOverlay);
+		}
+		if (SourceBiTangentOverlay && TargetBiTangentOverlay && SourceBiTangentOverlay->ElementCount() == TargetBiTangentOverlay->ElementCount())
+		{
+			check(ProxyBiTangentOverlay);
+			Private::InterpolateNormalOverlayValues(*SourceBiTangentOverlay, *TargetBiTangentOverlay, BlendAlpha, *ProxyBiTangentOverlay);
+		}
+		return true;
+	}
+}
\ No newline at end of file
diff --git a/Engine/Plugins/Experimental/MeshResizing/Source/MeshResizingCore/Private/MeshResizing/MeshResizingCoreModule.cpp b/Engine/Plugins/Experimental/MeshResizing/Source/MeshResizingCore/Private/MeshResizing/MeshResizingCoreModule.cpp
new file mode 100644
--- /dev/null
+++ b/Engine/Plugins/Experimental/MeshResizing/Source/MeshResizingCore/Private/MeshResizing/MeshResizingCoreModule.cpp
@@ -0,0 +1,20 @@
+// Copyright Epic Games, Inc. All Rights Reserved.
+
+#include "Modules/ModuleManager.h"
+
+namespace UE::MeshResizing
+{
+	class FMeshResizingCoreModule : public IModuleInterface
+	{
+	public:
+		virtual void StartupModule() override
+		{
+		}
+
+		virtual void ShutdownModule() override
+		{
+		}
+	};
+}
+
+IMPLEMENT_MODULE(UE::MeshResizing::FMeshResizingCoreModule, MeshResizingCore);
diff --git a/Engine/Plugins/Experimental/MeshResizing/Source/MeshResizingCore/Public/MeshResizing/BaseBodyTools.h b/Engine/Plugins/Experimental/MeshResizing/Source/MeshResizingCore/Public/MeshResizing/BaseBodyTools.h
new file mode 100644
--- /dev/null
+++ b/Engine/Plugins/Experimental/MeshResizing/Source/MeshResizingCore/Public/MeshResizing/BaseBodyTools.h
@@ -0,0 +1,30 @@
+// Copyright Epic Games, Inc. All Rights Reserved.
+
+#pragma once 
+#include "UObject/NameTypes.h"
+#include "Containers/Array.h"
+
+namespace UE::Geometry
+{
+	class FDynamicMesh3;
+}
+namespace UE::MeshResizing
+{
+	struct MESHRESIZINGCORE_API FBaseBodyTools
+	{
+		static const FName ImportedVertexVIDsAttrName;
+		static const FName RawPointIndicesVIDsAttrName;
+
+		/** @return success */
+		static bool AttachVertexMappingData(const FName& AttrName, const TArray<int32>& Data, UE::Geometry::FDynamicMesh3& Mesh);
+
+		/** @return success */
+		static bool GenerateResizableProxyFromVertexMappingData(const UE::Geometry::FDynamicMesh3& SourceMesh, const FName& SourceMappingName, const UE::Geometry::FDynamicMesh3& TargetMesh, const FName& TargetMappingName, UE::Geometry::FDynamicMesh3& ProxyMesh);
+
+		/** @return success */
+		static bool InterpolateResizableProxy(const UE::Geometry::FDynamicMesh3& SourceMesh, const UE::Geometry::FDynamicMesh3& TargetMesh, float BlendAlpha, UE::Geometry::FDynamicMesh3& ProxyMesh);
+
+
+	};
+}
+
diff --git a/Engine/Plugins/Experimental/MeshResizing/Source/MeshResizingNodes/MeshResizingDataflowNodes.Build.cs b/Engine/Plugins/Experimental/MeshResizing/Source/MeshResizingNodes/MeshResizingDataflowNodes.Build.cs
--- a/Engine/Plugins/Experimental/MeshResizing/Source/MeshResizingNodes/MeshResizingDataflowNodes.Build.cs
+++ b/Engine/Plugins/Experimental/MeshResizing/Source/MeshResizingNodes/MeshResizingDataflowNodes.Build.cs
@@ -27,7 +27,8 @@ public class MeshResizingDataflowNodes : ModuleRules
 				"DataflowEngine",
 				"DataflowNodes",
 				"GeometryFramework",
-				"MeshConversionEngineTypes"
+				"MeshConversionEngineTypes",
+				"MeshResizingCore"
 			}
 			);
 	}
diff --git a/Engine/Plugins/Experimental/MeshResizing/Source/MeshResizingNodes/Private/MeshResizing/BaseBodyDataflowNodes.cpp b/Engine/Plugins/Experimental/MeshResizing/Source/MeshResizingNodes/Private/MeshResizing/BaseBodyDataflowNodes.cpp
--- a/Engine/Plugins/Experimental/MeshResizing/Source/MeshResizingNodes/Private/MeshResizing/BaseBodyDataflowNodes.cpp
+++ b/Engine/Plugins/Experimental/MeshResizing/Source/MeshResizingNodes/Private/MeshResizing/BaseBodyDataflowNodes.cpp
@@ -1,9 +1,9 @@
 // Copyright Epic Games, Inc. All Rights Reserved.
 
 #include "MeshResizing/BaseBodyDataflowNodes.h"
+#include "MeshResizing/BaseBodyTools.h"
 #include "Dataflow/DataflowCore.h"
 #include "DynamicMesh/DynamicMesh3.h"
-#include "DynamicMesh/NonManifoldMappingSupport.h"
 #include "Engine/SkeletalMesh.h"
 #include "Rendering/SkeletalMeshRenderData.h"
 #include "Rendering/SkeletalMeshModel.h"
@@ -20,11 +20,14 @@ namespace UE::MeshResizing
 		static const FLinearColor CDefaultNodeBodyTintColor = FLinearColor(0.f, 0.f, 0.f, 0.5f);
 
 		DATAFLOW_NODE_REGISTER_CREATION_FACTORY(FSkeletalMeshToMeshDataflowNode);
+		DATAFLOW_NODE_REGISTER_CREATION_FACTORY(FGenerateResizableProxyDataflowNode);
+		DATAFLOW_NODE_REGISTER_CREATION_FACTORY(FGenerateInterpolatedProxyDataflowNode);
 	}
 }
 
 void FSkeletalMeshToMeshDataflowNode::Evaluate(UE::Dataflow::FContext& Context, const FDataflowOutput* Out) const
 {
+	using namespace UE::MeshResizing;
 	if (Out->IsA(&Mesh))
 	{
 		if (const USkeletalMesh* const InSkeletalMesh = GetValue(Context, &SkeletalMesh))
@@ -49,17 +52,23 @@ void FSkeletalMeshToMeshDataflowNode::Evaluate(UE::Dataflow::FContext& Context,
 						// Set ImportedVertices
 						TArray<int32> VertexToImportedVertexIDMap;
 						VertexToImportedVertexIDMap.SetNumUninitialized(DynMesh.VertexCount());
+						TArray<int32> VertexToRawPointIndicesMap;
+						VertexToRawPointIndicesMap.SetNumUninitialized(DynMesh.VertexCount());
 						check(DynMesh.VertexCount() == Converter.VertIDMap.Num());
 
 						const TArray<int32>& LODModelToImportVertexMap = SkeletalMeshLODModel->MeshToImportVertexMap;
+						const TArray<uint32>& RawPointIndices = SkeletalMeshLODModel->GetRawPointIndices();
 						
 						for (int32 DynMeshVertexIdx = 0; DynMeshVertexIdx < DynMesh.VertexCount(); ++DynMeshVertexIdx)
 						{
 							const int32 LODModelVertexIdx = Converter.VertIDMap[DynMeshVertexIdx];
 							const int32 ImportVertexIdx = LODModelToImportVertexMap[LODModelVertexIdx];
+							const int32 RawPointIdx = RawPointIndices[LODModelVertexIdx];
 							VertexToImportedVertexIDMap[DynMeshVertexIdx] = ImportVertexIdx;
+							VertexToRawPointIndicesMap[DynMeshVertexIdx] = RawPointIdx;
 						}
-						UE::Geometry::FNonManifoldMappingSupport::AttachNonManifoldVertexMappingData(VertexToImportedVertexIDMap, DynMesh);
+						FBaseBodyTools::AttachVertexMappingData(FBaseBodyTools::ImportedVertexVIDsAttrName, VertexToImportedVertexIDMap, DynMesh);
+						FBaseBodyTools::AttachVertexMappingData(FBaseBodyTools::RawPointIndicesVIDsAttrName, VertexToRawPointIndicesMap, DynMesh);
 
 						SetValue(Context, NewMesh, &Mesh);
 						return;
@@ -108,3 +117,78 @@ void FSkeletalMeshToMeshDataflowNode::Evaluate(UE::Dataflow::FContext& Context,
 		SetValue(Context, OutMaterials, &MaterialArray);
 	}
 }
+
+void FGenerateResizableProxyDataflowNode::Evaluate(UE::Dataflow::FContext& Context, const FDataflowOutput* Out) const
+{
+	if (Out->IsA(&TargetProxyMesh))
+	{
+		if (TObjectPtr<UDynamicMesh> InSourceMesh = GetValue<TObjectPtr<UDynamicMesh>>(Context, &SourceMesh))
+		{
+			if (TObjectPtr<UDynamicMesh> InTargetMesh = GetValue<TObjectPtr<UDynamicMesh>>(Context, &TargetMesh))
+			{
+				const UE::Geometry::FDynamicMesh3& SourceDynMesh = InSourceMesh->GetMeshRef();
+				const UE::Geometry::FDynamicMesh3& TargetDynMesh = InTargetMesh->GetMeshRef();
+
+				TObjectPtr<UDynamicMesh> NewTargetMesh = NewObject<UDynamicMesh>();
+				NewTargetMesh->Reset();
+				UE::Geometry::FDynamicMesh3& NewTargetDynMesh = NewTargetMesh->GetMeshRef();
+
+				if (UE::MeshResizing::FBaseBodyTools::GenerateResizableProxyFromVertexMappingData(SourceDynMesh, FName(*SourceMappingData), TargetDynMesh, FName(*TargetMappingData), NewTargetDynMesh))
+				{
+					SetValue(Context, NewTargetMesh, &TargetProxyMesh);
+					return;
+				}
+			}
+		}
+		SetValue(Context, TObjectPtr<UDynamicMesh>(NewObject<UDynamicMesh>()), &TargetProxyMesh);
+	}
+	else if (Out->IsA(&SourceProxyMesh))
+	{
+		SafeForwardInput(Context, &SourceMesh, &SourceProxyMesh);
+	}
+	else if (Out->IsA(&ProxyMaterialArray))
+	{
+		SafeForwardInput(Context, &SourceMaterialArray, &ProxyMaterialArray);
+	}
+}
+
+void FGenerateInterpolatedProxyDataflowNode::Evaluate(UE::Dataflow::FContext& Context, const FDataflowOutput* Out) const
+{
+	if (Out->IsA(&ProxyMesh))
+	{
+		if (BlendAlpha == 0.f)
+		{
+			SafeForwardInput(Context, &SourceMesh, &ProxyMesh);
+		}
+		else if (BlendAlpha == 1.f)
+		{
+			SafeForwardInput(Context, &TargetMesh, &ProxyMesh);
+		}
+		else
+		{
+			if (TObjectPtr<UDynamicMesh> InSourceMesh = GetValue<TObjectPtr<UDynamicMesh>>(Context, &SourceMesh))
+			{
+				if (TObjectPtr<UDynamicMesh> InTargetMesh = GetValue<TObjectPtr<UDynamicMesh>>(Context, &TargetMesh))
+				{
+					const UE::Geometry::FDynamicMesh3& SourceDynMesh = InSourceMesh->GetMeshRef();
+					const UE::Geometry::FDynamicMesh3& TargetDynMesh = InTargetMesh->GetMeshRef();
+
+					TObjectPtr<UDynamicMesh> NewTargetMesh = NewObject<UDynamicMesh>();
+					NewTargetMesh->Reset();
+					UE::Geometry::FDynamicMesh3& NewTargetDynMesh = NewTargetMesh->GetMeshRef();
+
+					if (UE::MeshResizing::FBaseBodyTools::InterpolateResizableProxy(SourceDynMesh, TargetDynMesh, BlendAlpha, NewTargetDynMesh))
+					{
+						SetValue(Context, NewTargetMesh, &ProxyMesh);
+						return;
+					}
+				}
+			}
+			SetValue(Context, TObjectPtr<UDynamicMesh>(NewObject<UDynamicMesh>()), &ProxyMesh);
+		}
+	}
+	else if (Out->IsA(&ProxyMaterialArray))
+	{
+		SafeForwardInput(Context, &SourceMaterialArray, &ProxyMaterialArray);
+	}
+}
\ No newline at end of file
diff --git a/Engine/Plugins/Experimental/MeshResizing/Source/MeshResizingNodes/Public/MeshResizing/BaseBodyDataflowNodes.h b/Engine/Plugins/Experimental/MeshResizing/Source/MeshResizingNodes/Public/MeshResizing/BaseBodyDataflowNodes.h
--- a/Engine/Plugins/Experimental/MeshResizing/Source/MeshResizingNodes/Public/MeshResizing/BaseBodyDataflowNodes.h
+++ b/Engine/Plugins/Experimental/MeshResizing/Source/MeshResizingNodes/Public/MeshResizing/BaseBodyDataflowNodes.h
@@ -17,8 +17,8 @@ class UMaterialInterface;
  * Converts a SkeletalMesh into a DynamicMesh with Imported Vertex information
  *
  */
-USTRUCT(meta = (MeshResizing))
-struct FSkeletalMeshToMeshDataflowNode : public FDataflowNode
+USTRUCT(meta = (MeshResizing, Experimental))
+struct FSkeletalMeshToMeshDataflowNode final : public FDataflowNode
 {
 	GENERATED_USTRUCT_BODY()
 	DATAFLOW_NODE_DEFINE_INTERNAL(FSkeletalMeshToMeshDataflowNode, "SkeletalMeshToMesh", "Mesh|Utilities", "")
@@ -59,6 +59,127 @@ private:
 	bool bRecordImportedVertices = true;
 };
 
+
+/**
+ *
+ * Generate a pair of Dynamic Meshes with the same topology that can be interpolated.
+ * 
+ * Currently, this node relies on the vertex mapping data existing on the input source and target meshes,
+ * and that the mapped vertices on both meshes match.
+ */
+USTRUCT(meta = (MeshResizing, Experimental))
+struct FGenerateResizableProxyDataflowNode final : public FDataflowNode
+{
+	GENERATED_USTRUCT_BODY()
+	DATAFLOW_NODE_DEFINE_INTERNAL(FGenerateResizableProxyDataflowNode, "GenerateResizableProxy", "MeshResizing", "")
+	DATAFLOW_NODE_RENDER_TYPE("SurfaceRender",FName("FDynamicMesh3"), {"TargetProxyMesh", "ProxyMaterialArray"})
+
+public:
+
+	/** Source mesh */
+	UPROPERTY(meta = (DataflowInput))
+	TObjectPtr<UDynamicMesh> SourceMesh;
+
+	/** Source materials.*/
+	UPROPERTY(meta = (DataflowInput))
+	TArray<TObjectPtr<UMaterialInterface>> SourceMaterialArray;
+
+	/** Target mesh */
+	UPROPERTY(meta = (DataflowInput))
+	TObjectPtr<UDynamicMesh> TargetMesh;
+
+	/** Output source proxy mesh */
+	UPROPERTY(meta = (DataflowOutput, DataflowPassthrough = "SourceMesh"))
+	TObjectPtr<UDynamicMesh> SourceProxyMesh;
+
+	/** Output source proxy mesh */
+	UPROPERTY(meta = (DataflowOutput, DataflowPassthrough = "SourceMesh"))
+	TObjectPtr<UDynamicMesh> TargetProxyMesh;
+
+	/** Target materials.*/
+	UPROPERTY(meta = (DataflowOutput, DataflowPassthrough = "SourceMaterialArray"))
+	TArray<TObjectPtr<UMaterialInterface>> ProxyMaterialArray;
+
+
+	FGenerateResizableProxyDataflowNode(const UE::Dataflow::FNodeParameters& InParam, FGuid InGuid = FGuid::NewGuid())
+		: FDataflowNode(InParam, InGuid)
+	{
+		RegisterInputConnection(&SourceMesh);
+		RegisterInputConnection(&SourceMaterialArray);
+		RegisterInputConnection(&TargetMesh);
+		RegisterOutputConnection(&SourceProxyMesh, &SourceMesh);
+		RegisterOutputConnection(&TargetProxyMesh, &SourceMesh);
+		RegisterOutputConnection(&ProxyMaterialArray, &SourceMaterialArray);
+	}
+
+private:
+
+	/** Source vertex mapping data. TODO: only have two choices that work currently. Make this an enum or something */
+	UPROPERTY(EditAnywhere, Category = "Resizable Proxy")
+	FString SourceMappingData = "ImportedVertexVIDsAttr";
+
+	/** Target vertex mapping data. TODO: only have two choices that work currently. Make this an enum or something */
+	UPROPERTY(EditAnywhere, Category = "Resizable Proxy")
+	FString TargetMappingData = "ImportedVertexVIDsAttr";
+
+	virtual void Evaluate(UE::Dataflow::FContext& Context, const FDataflowOutput* Out) const override;
+};
+/**
+ *
+ * Generate a pair of Dynamic Meshes with the same topology that can be interpolated.
+ *
+ * Currently, this node relies on the vertex mapping data existing on the input source and target meshes,
+ * and that the mapped vertices on both meshes match.
+ */
+USTRUCT(meta = (MeshResizing, Experimental))
+struct FGenerateInterpolatedProxyDataflowNode final : public FDataflowNode
+{
+	GENERATED_USTRUCT_BODY()
+	DATAFLOW_NODE_DEFINE_INTERNAL(FGenerateInterpolatedProxyDataflowNode, "GenerateInterpolatedProxy", "MeshResizing", "")
+	DATAFLOW_NODE_RENDER_TYPE("SurfaceRender", FName("FDynamicMesh3"), { "ProxyMesh", "ProxyMaterialArray" })
+
+public:
+
+	/** Source mesh */
+	UPROPERTY(meta = (DataflowInput))
+	TObjectPtr<UDynamicMesh> SourceMesh;
+
+	/** Source materials.*/
+	UPROPERTY(meta = (DataflowInput))
+	TArray<TObjectPtr<UMaterialInterface>> SourceMaterialArray;
+
+	/** Target mesh */
+	UPROPERTY(meta = (DataflowInput))
+	TObjectPtr<UDynamicMesh> TargetMesh;
+
+	/** Output proxy mesh */
+	UPROPERTY(meta = (DataflowOutput, DataflowPassthrough = "SourceMesh"))
+	TObjectPtr<UDynamicMesh> ProxyMesh;
+
+	/** Proxy materials.*/
+	UPROPERTY(meta = (DataflowOutput, DataflowPassthrough = "SourceMaterialArray"))
+	TArray<TObjectPtr<UMaterialInterface>> ProxyMaterialArray;
+
+
+	FGenerateInterpolatedProxyDataflowNode(const UE::Dataflow::FNodeParameters& InParam, FGuid InGuid = FGuid::NewGuid())
+		: FDataflowNode(InParam, InGuid)
+	{
+		RegisterInputConnection(&SourceMesh);
+		RegisterInputConnection(&SourceMaterialArray);
+		RegisterInputConnection(&TargetMesh);
+		RegisterOutputConnection(&ProxyMesh, &SourceMesh);
+		RegisterOutputConnection(&ProxyMaterialArray, &SourceMaterialArray);
+	}
+
+private:
+
+	/** Alpha between source (0) and target (1) */
+	UPROPERTY(EditAnywhere, Category = "Interpolate Proxy", meta = (ClampMin = "0", ClampMax = "1"))
+	float BlendAlpha = 1.f;
+
+	virtual void Evaluate(UE::Dataflow::FContext& Context, const FDataflowOutput* Out) const override;
+};
+
 namespace UE::MeshResizing
 {
 	void RegisterBaseBodyDataflowNodes();
