MeshResizing: Add RBFInterpolation (from the HairStrands plugin) option for MeshWarp node.

Exposed GroomBinding_RBFWeighting::FWeightsBuilder to HairStrandsCore API.

diff --git a/Engine/Plugins/Experimental/MeshResizing/MeshResizing.uplugin b/Engine/Plugins/Experimental/MeshResizing/MeshResizing.uplugin
--- a/Engine/Plugins/Experimental/MeshResizing/MeshResizing.uplugin
+++ b/Engine/Plugins/Experimental/MeshResizing/MeshResizing.uplugin
@@ -34,6 +34,10 @@
 		{
 			"Name": "GeometryCollectionPlugin",
 			"Enabled": true
+		},
+		{
+			"Name": "HairStrands",
+			"Enabled":  true
 		}
 	]
 }
diff --git a/Engine/Plugins/Experimental/MeshResizing/Source/MeshResizingCore/MeshResizingCore.Build.cs b/Engine/Plugins/Experimental/MeshResizing/Source/MeshResizingCore/MeshResizingCore.Build.cs
--- a/Engine/Plugins/Experimental/MeshResizing/Source/MeshResizingCore/MeshResizingCore.Build.cs
+++ b/Engine/Plugins/Experimental/MeshResizing/Source/MeshResizingCore/MeshResizingCore.Build.cs
@@ -13,6 +13,7 @@ namespace UnrealBuildTool.Rules
 				{
 					"Core",
 					"GeometryCore",
+					"HairStrandsCore",
 				}
 			);
 		}
diff --git a/Engine/Plugins/Experimental/MeshResizing/Source/MeshResizingCore/Private/MeshResizing/RBFInterpolation.cpp b/Engine/Plugins/Experimental/MeshResizing/Source/MeshResizingCore/Private/MeshResizing/RBFInterpolation.cpp
new file mode 100644
--- /dev/null
+++ b/Engine/Plugins/Experimental/MeshResizing/Source/MeshResizingCore/Private/MeshResizing/RBFInterpolation.cpp
@@ -0,0 +1,92 @@
+// Copyright Epic Games, Inc. All Rights Reserved.
+
+#include "MeshResizing/RBFInterpolation.h"
+#include "DynamicMesh/DynamicMesh3.h"
+#include "GroomBindingBuilder.h"
+#include "GroomRBFDeformer.h"
+
+namespace UE::MeshResizing
+{
+	namespace Private
+	{
+		static void UpdateInterpolationData(const GroomBinding_RBFWeighting::FPointsSampler& PointsSampler, const GroomBinding_RBFWeighting::FWeightsBuilder& InterpolationWeights, FRBFInterpolation::FInterpolationData& OutData)
+		{
+			OutData.SampleIndices = PointsSampler.SampleIndices;
+			OutData.SampleRestPositions = PointsSampler.SamplePositions;
+			OutData.InterpolationWeights = InterpolationWeights.InverseEntries;
+		}
+
+		// From GroomRBFDeformer
+		static void UpdateMeshSamples(const TArray<FVector3f>& TargetPositions, const FRBFInterpolation::FInterpolationData& InterpolationData, TArray<FVector3f>& OutSampleDeformations)
+		{
+			const uint32 MaxSampleCount = TargetPositions.Num();
+			check(MaxSampleCount == InterpolationData.SampleIndices.Num());
+			check(MaxSampleCount == InterpolationData.SampleRestPositions.Num());
+
+			const uint32 EntryCount = FGroomRBFDeformer::GetEntryCount(MaxSampleCount);
+			OutSampleDeformations.SetNum(EntryCount);
+
+			for (uint32 SampleIndex = 0; SampleIndex < EntryCount; ++SampleIndex)
+			{
+				uint32 WeightsOffset = SampleIndex * EntryCount;
+				FVector3f SampleDeformation(FVector3f::ZeroVector);
+				for (uint32 Index = 0; Index < MaxSampleCount; ++Index, ++WeightsOffset)
+				{
+					SampleDeformation += InterpolationData.InterpolationWeights[WeightsOffset] * (TargetPositions[Index] - InterpolationData.SampleRestPositions[Index]);
+				}
+				OutSampleDeformations[SampleIndex] = SampleDeformation;
+			}
+		}
+
+		// From GroomRBFDeformer
+		static void ApplyRBFDeformation(const FRBFInterpolation::FInterpolationData& InterpolationData, const TArray<FVector3f>& SampleDeformations, UE::Geometry::FDynamicMesh3& DeformingMesh)
+		{
+			const uint32 SampleCount = InterpolationData.SampleRestPositions.Num();
+			check(SampleCount + 4 == SampleDeformations.Num());
+			for (const int32 VertexIndex : DeformingMesh.VertexIndicesItr())
+			{
+				const FVector3f RestControlPoint(FVector3f(DeformingMesh.GetVertexRef(VertexIndex)));
+				FVector3f ControlPoint = RestControlPoint;
+				for (uint32 SampleIndex = 0; SampleIndex < SampleCount; ++SampleIndex)
+				{
+					const FVector3f PositionDelta = RestControlPoint - InterpolationData.SampleRestPositions[SampleIndex];
+					const float FunctionValue = FMath::Sqrt(FVector3f::DotProduct(PositionDelta, PositionDelta) + 1);
+					ControlPoint += FunctionValue * SampleDeformations[SampleIndex];
+				}
+				ControlPoint += SampleDeformations[SampleCount];
+				ControlPoint += SampleDeformations[SampleCount + 1] * RestControlPoint.X;
+				ControlPoint += SampleDeformations[SampleCount + 2] * RestControlPoint.Y;
+				ControlPoint += SampleDeformations[SampleCount + 3] * RestControlPoint.Z;
+
+				DeformingMesh.SetVertex(VertexIndex, FVector3d(ControlPoint));
+			}
+		}		
+	}
+
+	void FRBFInterpolation::GenerateWeights(const UE::Geometry::FDynamicMesh3& BaseMesh, int32 NumInterpolationPoints, FInterpolationData& OutData)
+	{
+		// Groom binding sampler expects points to be a contiguous array of FVector3f. Just copy that way for now.
+		TArray<FVector3f> Positions;
+		Positions.Reserve(BaseMesh.VertexCount());
+		for (const int32 VertexIndex : BaseMesh.VertexIndicesItr())
+		{
+			Positions.Emplace(FVector3f(BaseMesh.GetVertexRef(VertexIndex)));
+		}
+
+		// Consider all points to be valid for now
+		TArray<bool> ValidPoints;
+		ValidPoints.Init(true, BaseMesh.VertexCount());
+
+		GroomBinding_RBFWeighting::FPointsSampler PointsSampler(ValidPoints, Positions.GetData(), NumInterpolationPoints);
+		const uint32 SampleCount = PointsSampler.SamplePositions.Num();
+		GroomBinding_RBFWeighting::FWeightsBuilder InterpolationWeights(SampleCount, SampleCount, PointsSampler.SamplePositions.GetData(), PointsSampler.SamplePositions.GetData());
+		Private::UpdateInterpolationData(PointsSampler, InterpolationWeights, OutData);
+	}
+
+	void FRBFInterpolation::DeformPoints(const TArray<FVector3f>& TargetPositions, const FInterpolationData& InterpolationData, UE::Geometry::FDynamicMesh3& DeformingMesh)
+	{
+		TArray<FVector3f> SampleDeformations;
+		Private::UpdateMeshSamples(TargetPositions, InterpolationData, SampleDeformations);
+		Private::ApplyRBFDeformation(InterpolationData, SampleDeformations, DeformingMesh);
+	}
+}
\ No newline at end of file
diff --git a/Engine/Plugins/Experimental/MeshResizing/Source/MeshResizingCore/Public/MeshResizing/RBFInterpolation.h b/Engine/Plugins/Experimental/MeshResizing/Source/MeshResizingCore/Public/MeshResizing/RBFInterpolation.h
new file mode 100644
--- /dev/null
+++ b/Engine/Plugins/Experimental/MeshResizing/Source/MeshResizingCore/Public/MeshResizing/RBFInterpolation.h
@@ -0,0 +1,27 @@
+// Copyright Epic Games, Inc. All Rights Reserved.
+
+#pragma once 
+#include "UObject/NameTypes.h"
+#include "Containers/Array.h"
+
+namespace UE::Geometry
+{
+	class FDynamicMesh3;
+}
+namespace UE::MeshResizing
+{
+	struct MESHRESIZINGCORE_API FRBFInterpolation
+	{
+
+		struct FInterpolationData
+		{
+			TArray<uint32> SampleIndices;
+			TArray<FVector3f> SampleRestPositions;
+			TArray<float> InterpolationWeights;
+		};
+
+		static void GenerateWeights(const UE::Geometry::FDynamicMesh3& BaseMesh, int32 NumInterpolationPoints, FInterpolationData& OutData);
+		static void DeformPoints(const TArray<FVector3f>& TargetPositions, const FInterpolationData& InterpolationData, UE::Geometry::FDynamicMesh3& DeformingMesh);
+	};
+}
+
diff --git a/Engine/Plugins/Experimental/MeshResizing/Source/MeshResizingNodes/Private/MeshResizing/BaseBodyDataflowNodes.h b/Engine/Plugins/Experimental/MeshResizing/Source/MeshResizingNodes/Private/MeshResizing/BaseBodyDataflowNodes.h
new file mode 100644
--- /dev/null
+++ b/Engine/Plugins/Experimental/MeshResizing/Source/MeshResizingNodes/Private/MeshResizing/BaseBodyDataflowNodes.h
@@ -0,0 +1,187 @@
+// Copyright Epic Games, Inc. All Rights Reserved.
+
+#pragma once 
+
+#include "Dataflow/DataflowEngine.h"
+#include "DynamicMesh/DynamicMesh3.h"
+#include "DynamicMeshProcessor.h"
+
+#include "BaseBodyDataflowNodes.generated.h"
+
+class USkeletalMesh;
+class UDynamicMesh;
+class UMaterialInterface;
+
+/**
+ *
+ * Converts a SkeletalMesh into a DynamicMesh with Imported Vertex information
+ *
+ */
+USTRUCT(meta = (MeshResizing, Experimental))
+struct FSkeletalMeshToMeshDataflowNode final : public FDataflowNode
+{
+	GENERATED_USTRUCT_BODY()
+	DATAFLOW_NODE_DEFINE_INTERNAL(FSkeletalMeshToMeshDataflowNode, "SkeletalMeshToMesh", "Mesh|Utilities", "")
+	DATAFLOW_NODE_RENDER_TYPE("SurfaceRender",FName("FDynamicMesh3"), {"Mesh", "MaterialArray"})
+
+public:
+
+	/** SkeletalMesh to convert */
+	UPROPERTY(EditAnywhere, Category = "SkeletalMesh", meta = (DataflowInput));
+	TObjectPtr<const USkeletalMesh> SkeletalMesh;
+
+	/** Output mesh */
+	UPROPERTY(meta = (DataflowOutput))
+	TObjectPtr<UDynamicMesh> Mesh;
+
+	/** Output materials */
+	UPROPERTY(meta = (DataflowOutput))
+	TArray<TObjectPtr<UMaterialInterface>> MaterialArray;
+
+	FSkeletalMeshToMeshDataflowNode(const UE::Dataflow::FNodeParameters& InParam, FGuid InGuid = FGuid::NewGuid())
+		: FDataflowNode(InParam, InGuid)
+	{
+		RegisterInputConnection(&SkeletalMesh);
+		RegisterOutputConnection(&Mesh);
+		RegisterOutputConnection(&MaterialArray);
+	}
+
+private:
+
+	virtual void Evaluate(UE::Dataflow::FContext& Context, const FDataflowOutput* Out) const override;
+
+	/** Specifies the LOD level to use */
+	UPROPERTY(EditAnywhere, Category = "SkeletalMesh", meta = (DisplayName = "LOD Level"));
+	int32 LODLevel = 0;
+
+	/** Record ImportedVertices (if available) as NonManifold mapping data. This requires Editor-Only data.*/
+	UPROPERTY(EditAnywhere, Category = "SkeletalMesh")
+	bool bRecordImportedVertices = true;
+};
+
+
+/**
+ *
+ * Generate a pair of Dynamic Meshes with the same topology that can be interpolated.
+ * 
+ * Currently, this node relies on the vertex mapping data existing on the input source and target meshes,
+ * and that the mapped vertices on both meshes match.
+ */
+USTRUCT(meta = (MeshResizing, Experimental))
+struct FGenerateResizableProxyDataflowNode final : public FDataflowNode
+{
+	GENERATED_USTRUCT_BODY()
+	DATAFLOW_NODE_DEFINE_INTERNAL(FGenerateResizableProxyDataflowNode, "GenerateResizableProxy", "MeshResizing", "")
+	DATAFLOW_NODE_RENDER_TYPE("SurfaceRender",FName("FDynamicMesh3"), {"TargetProxyMesh", "ProxyMaterialArray"})
+
+public:
+
+	/** Source mesh */
+	UPROPERTY(meta = (DataflowInput))
+	TObjectPtr<UDynamicMesh> SourceMesh;
+
+	/** Source materials.*/
+	UPROPERTY(meta = (DataflowInput))
+	TArray<TObjectPtr<UMaterialInterface>> SourceMaterialArray;
+
+	/** Target mesh */
+	UPROPERTY(meta = (DataflowInput))
+	TObjectPtr<UDynamicMesh> TargetMesh;
+
+	/** Output source proxy mesh */
+	UPROPERTY(meta = (DataflowOutput, DataflowPassthrough = "SourceMesh"))
+	TObjectPtr<UDynamicMesh> SourceProxyMesh;
+
+	/** Output source proxy mesh */
+	UPROPERTY(meta = (DataflowOutput, DataflowPassthrough = "SourceMesh"))
+	TObjectPtr<UDynamicMesh> TargetProxyMesh;
+
+	/** Target materials.*/
+	UPROPERTY(meta = (DataflowOutput, DataflowPassthrough = "SourceMaterialArray"))
+	TArray<TObjectPtr<UMaterialInterface>> ProxyMaterialArray;
+
+
+	FGenerateResizableProxyDataflowNode(const UE::Dataflow::FNodeParameters& InParam, FGuid InGuid = FGuid::NewGuid())
+		: FDataflowNode(InParam, InGuid)
+	{
+		RegisterInputConnection(&SourceMesh);
+		RegisterInputConnection(&SourceMaterialArray);
+		RegisterInputConnection(&TargetMesh);
+		RegisterOutputConnection(&SourceProxyMesh, &SourceMesh);
+		RegisterOutputConnection(&TargetProxyMesh, &SourceMesh);
+		RegisterOutputConnection(&ProxyMaterialArray, &SourceMaterialArray);
+	}
+
+private:
+
+	/** Source vertex mapping data. TODO: only have two choices that work currently. Make this an enum or something */
+	UPROPERTY(EditAnywhere, Category = "Resizable Proxy")
+	FString SourceMappingData = "ImportedVertexVIDsAttr";
+
+	/** Target vertex mapping data. TODO: only have two choices that work currently. Make this an enum or something */
+	UPROPERTY(EditAnywhere, Category = "Resizable Proxy")
+	FString TargetMappingData = "ImportedVertexVIDsAttr";
+
+	virtual void Evaluate(UE::Dataflow::FContext& Context, const FDataflowOutput* Out) const override;
+};
+/**
+ *
+ * Generate a pair of Dynamic Meshes with the same topology that can be interpolated.
+ *
+ * Currently, this node relies on the vertex mapping data existing on the input source and target meshes,
+ * and that the mapped vertices on both meshes match.
+ */
+USTRUCT(meta = (MeshResizing, Experimental))
+struct FGenerateInterpolatedProxyDataflowNode final : public FDataflowNode
+{
+	GENERATED_USTRUCT_BODY()
+	DATAFLOW_NODE_DEFINE_INTERNAL(FGenerateInterpolatedProxyDataflowNode, "GenerateInterpolatedProxy", "MeshResizing", "")
+	DATAFLOW_NODE_RENDER_TYPE("SurfaceRender", FName("FDynamicMesh3"), { "ProxyMesh", "ProxyMaterialArray" })
+
+public:
+
+	/** Source mesh */
+	UPROPERTY(meta = (DataflowInput))
+	TObjectPtr<UDynamicMesh> SourceMesh;
+
+	/** Source materials.*/
+	UPROPERTY(meta = (DataflowInput))
+	TArray<TObjectPtr<UMaterialInterface>> SourceMaterialArray;
+
+	/** Target mesh */
+	UPROPERTY(meta = (DataflowInput))
+	TObjectPtr<UDynamicMesh> TargetMesh;
+
+	/** Output proxy mesh */
+	UPROPERTY(meta = (DataflowOutput, DataflowPassthrough = "SourceMesh"))
+	TObjectPtr<UDynamicMesh> ProxyMesh;
+
+	/** Proxy materials.*/
+	UPROPERTY(meta = (DataflowOutput, DataflowPassthrough = "SourceMaterialArray"))
+	TArray<TObjectPtr<UMaterialInterface>> ProxyMaterialArray;
+
+
+	FGenerateInterpolatedProxyDataflowNode(const UE::Dataflow::FNodeParameters& InParam, FGuid InGuid = FGuid::NewGuid())
+		: FDataflowNode(InParam, InGuid)
+	{
+		RegisterInputConnection(&SourceMesh);
+		RegisterInputConnection(&SourceMaterialArray);
+		RegisterInputConnection(&TargetMesh);
+		RegisterOutputConnection(&ProxyMesh, &SourceMesh);
+		RegisterOutputConnection(&ProxyMaterialArray, &SourceMaterialArray);
+	}
+
+private:
+
+	/** Alpha between source (0) and target (1) */
+	UPROPERTY(EditAnywhere, Category = "Interpolate Proxy", meta = (ClampMin = "0", ClampMax = "1"))
+	float BlendAlpha = 1.f;
+
+	virtual void Evaluate(UE::Dataflow::FContext& Context, const FDataflowOutput* Out) const override;
+};
+
+namespace UE::MeshResizing
+{
+	void RegisterBaseBodyDataflowNodes();
+}
+
diff --git a/Engine/Plugins/Experimental/MeshResizing/Source/MeshResizingNodes/Private/MeshResizing/MeshWarpNode.cpp b/Engine/Plugins/Experimental/MeshResizing/Source/MeshResizingNodes/Private/MeshResizing/MeshWarpNode.cpp
--- a/Engine/Plugins/Experimental/MeshResizing/Source/MeshResizingNodes/Private/MeshResizing/MeshWarpNode.cpp
+++ b/Engine/Plugins/Experimental/MeshResizing/Source/MeshResizingNodes/Private/MeshResizing/MeshWarpNode.cpp
@@ -1,5 +1,6 @@
 // Copyright Epic Games, Inc. All Rights Reserved.
 #include "MeshResizing/MeshWarpNode.h"
+#include "MeshResizing/RBFInterpolation.h"
 #include "DynamicMesh/DynamicMesh3.h"
 #include "DynamicMesh/DynamicMeshAABBTree3.h"
 #include "UDynamicMesh.h"
@@ -8,56 +9,40 @@
 
 #include UE_INLINE_GENERATED_CPP_BY_NAME(MeshWarpNode)
 
-namespace  UE::Chaos::MeshResizing::Private
+namespace  UE::MeshResizing::Private
 {
 	void BlendMesh(
-		FDynamicMesh3* OutMesh,
-		const UDynamicMesh* SourceUMesh,
-		const UDynamicMesh* TargetUMesh,
+		FDynamicMesh3& OutMesh,
+		const FDynamicMesh3& SourceMesh,
+		const FDynamicMesh3& TargetMesh,
 		float Alpha)
 	{
 		using namespace UE::Geometry;
 
-		if (OutMesh == nullptr)
-		{
-			UE_LOG(LogTemp, Warning, TEXT("BlendMesh: OutMesh is Null"));
-			return;
-		}
-
-		const FDynamicMesh3* const SourceMesh = SourceUMesh->GetMeshPtr();
-		OutMesh->Copy(*SourceMesh);
-
-		const FDynamicMesh3* const TargetMesh = TargetUMesh->GetMeshPtr();
+		OutMesh.Copy(SourceMesh);
 
-		ensure(SourceMesh->VertexCount() == TargetMesh->VertexCount());
-		ensure(SourceMesh->VertexCount() == OutMesh->VertexCount());
+		ensure(SourceMesh.VertexCount() == TargetMesh.VertexCount());
+		ensure(SourceMesh.VertexCount() == OutMesh.VertexCount());
 
-		for (const int32 VertexIndex : TargetMesh->VertexIndicesItr())
+		for (const int32 VertexIndex : TargetMesh.VertexIndicesItr())
 		{
-			const FVector3d SourceVert = SourceMesh->GetVertex(VertexIndex);
-			const FVector3d TargetVert = TargetMesh->GetVertex(VertexIndex);
+			const FVector3d SourceVert = SourceMesh.GetVertex(VertexIndex);
+			const FVector3d TargetVert = TargetMesh.GetVertex(VertexIndex);
 			const FVector3d FinalPoint = (1.0 - Alpha) * SourceVert + Alpha * TargetVert;
-			OutMesh->SetVertex(VertexIndex, FinalPoint);
+			OutMesh.SetVertex(VertexIndex, FinalPoint);
 		}
 	}
 
 
-	void ApplyRefitToMesh(
-		FDynamicMesh3* ResizedMesh,
-		const UDynamicMesh* SourceUMesh,
+	void ApplyRefitToMesh_WrapDeform(
+		FDynamicMesh3& ResizedMesh,
+		const FDynamicMesh3& SourceMesh,
 		const FDynamicMesh3& TargetMesh,
 		float Alpha)
 	{
 		using namespace UE::Geometry;
 
-		if (ResizedMesh == nullptr)
-		{
-			UE_LOG(LogTemp, Warning, TEXT("ApplyRefitToMesh: ResizedMesh is Null"));
-			return;
-		}
-
-		const FDynamicMesh3* const SourceMesh = SourceUMesh->GetMeshPtr();
-		FDynamicMeshAABBTree3 SourceTree(SourceMesh);
+		FDynamicMeshAABBTree3 SourceTree(&SourceMesh);
 
 		struct FEmbeddingInfo
 		{
@@ -66,21 +51,21 @@ namespace  UE::Chaos::MeshResizing::Private
 		};
 
 		TArray<FEmbeddingInfo> EmbeddingInfo;
-		EmbeddingInfo.SetNum(ResizedMesh->VertexCount());
+		EmbeddingInfo.SetNum(ResizedMesh.VertexCount());
 
-		for (const int32 VertexIndex : ResizedMesh->VertexIndicesItr())
+		for (const int32 VertexIndex : ResizedMesh.VertexIndicesItr())
 		{
 			EmbeddingInfo[VertexIndex].ClosestTriangle = INDEX_NONE;
 
-			const FVector3d VertexLocation = ResizedMesh->GetVertex(VertexIndex);
+			const FVector3d VertexLocation = ResizedMesh.GetVertex(VertexIndex);
 
 			double NearestDistSqr;
 			const int32 NearTriID = SourceTree.FindNearestTriangle(VertexLocation, NearestDistSqr);
 			if (NearTriID >= 0)
 			{
-				const FVector3d Normal = SourceMesh->GetTriNormal(NearTriID);
+				const FVector3d Normal = SourceMesh.GetTriNormal(NearTriID);
 				FVector3d TriPointA, TriPointB, TriPointC;
-				SourceMesh->GetTriVertices(NearTriID, TriPointA, TriPointB, TriPointC);
+				SourceMesh.GetTriVertices(NearTriID, TriPointA, TriPointB, TriPointC);
 				const FPlane3d TriPlane(Normal, TriPointA);
 				const double Dist = TriPlane.DistanceTo(VertexLocation);
 				const FVector3d PointOnPlane = VertexLocation - Dist * Normal;
@@ -91,23 +76,43 @@ namespace  UE::Chaos::MeshResizing::Private
 			}
 		}
 
-		for (const int32 VertexIndex : ResizedMesh->VertexIndicesItr())
+		for (const int32 VertexIndex : ResizedMesh.VertexIndicesItr())
 		{
 			if (EmbeddingInfo[VertexIndex].ClosestTriangle != INDEX_NONE)
 			{
-				const FVector3d Normal = SourceMesh->GetTriNormal(EmbeddingInfo[VertexIndex].ClosestTriangle);
+				const FVector3d Normal = SourceMesh.GetTriNormal(EmbeddingInfo[VertexIndex].ClosestTriangle);
 				FVector3d TriPointA, TriPointB, TriPointC;
 				TargetMesh.GetTriVertices(EmbeddingInfo[VertexIndex].ClosestTriangle, TriPointA, TriPointB, TriPointC);
 
 				const FVector4d& Weights = EmbeddingInfo[VertexIndex].Weights;
 				const FVector PointOnPlane = Weights[0] * TriPointA + Weights[1] * TriPointB + Weights[2] * TriPointC;
-				const FVector3d CurrentVertexPosition = ResizedMesh->GetVertex(VertexIndex);
+				const FVector3d CurrentVertexPosition = ResizedMesh.GetVertex(VertexIndex);
 				const FVector FinalPoint = (1.0 - Alpha) * CurrentVertexPosition + Alpha * (PointOnPlane + Weights[3] * Normal);
 
-				ResizedMesh->SetVertex(VertexIndex, FinalPoint);
+				ResizedMesh.SetVertex(VertexIndex, FinalPoint);
 			}
 		}
 	}
+
+	void ApplyRefitToMesh_RBFInterpolate(
+		FDynamicMesh3& ResizedMesh,
+		const FDynamicMesh3& SourceMesh,
+		const FDynamicMesh3& TargetMesh,
+		int32 NumInterpolationPoints)
+	{
+		FRBFInterpolation::FInterpolationData InterpolationData;
+		FRBFInterpolation::GenerateWeights(SourceMesh, NumInterpolationPoints, InterpolationData);
+
+		// Get target mesh samples
+		TArray<FVector3f> TargetSamplePoints;
+		TargetSamplePoints.Reserve(InterpolationData.SampleIndices.Num());
+		for (const uint32 SampleIndex : InterpolationData.SampleIndices)
+		{
+			TargetSamplePoints.Emplace(FVector3f(TargetMesh.GetVertex(SampleIndex)));
+		}
+		// Do the interpolation
+		FRBFInterpolation::DeformPoints(TargetSamplePoints, InterpolationData, ResizedMesh);
+	}
 }
 
 
@@ -134,23 +139,38 @@ void FMeshWarpNode::Evaluate(UE::Dataflow::FContext& Context, const FDataflowOut
 
 	if (TObjectPtr<UDynamicMesh> InMeshToResize = GetValue(Context, &MeshToResize))
 	{
-		if (TObjectPtr<UDynamicMesh> InSourceMesh = GetValue(Context, &SourceMesh))
+		if (TObjectPtr<UDynamicMesh> InSourceUMesh = GetValue(Context, &SourceMesh))
 		{
-			if (TObjectPtr<UDynamicMesh> InTargetMesh = GetValue(Context, &TargetMesh))
+			if(const FDynamicMesh3* const InSourceMesh = InSourceUMesh->GetMeshPtr())
 			{
-				FDynamicMesh3 BlendedTargetFMesh;
-				UE::Chaos::MeshResizing::Private::BlendMesh(&BlendedTargetFMesh, InSourceMesh, InTargetMesh, Alpha);
-
-				FDynamicMesh3 ResizeMesh;
-				ResizeMesh.Copy(InMeshToResize->GetMeshRef());
-				UE::Chaos::MeshResizing::Private::ApplyRefitToMesh(&ResizeMesh, InSourceMesh, BlendedTargetFMesh, 1.0f);
-
-				OutBlendedTargetMesh->SetMesh(MoveTemp(BlendedTargetFMesh));
-				OutResizedMesh->SetMesh(MoveTemp(ResizeMesh));
-
-				if (IsConnected(&MeshToResizeMaterialArray))
+				if (TObjectPtr<UDynamicMesh> InTargetUMesh = GetValue(Context, &TargetMesh))
 				{
-					OutMaterials = GetValue(Context, &MeshToResizeMaterialArray);
+					if (const FDynamicMesh3* const InTargetMesh = InTargetUMesh->GetMeshPtr())
+					{
+						FDynamicMesh3 BlendedTargetFMesh;
+						UE::MeshResizing::Private::BlendMesh(BlendedTargetFMesh, *InSourceMesh, *InTargetMesh, Alpha);
+
+						FDynamicMesh3 ResizeMesh;
+						ResizeMesh.Copy(InMeshToResize->GetMeshRef());
+						switch (WarpMethod)
+						{
+						case EMeshResizingWarpMethod::RBFInterpolate:
+							UE::MeshResizing::Private::ApplyRefitToMesh_RBFInterpolate(ResizeMesh, *InSourceMesh, BlendedTargetFMesh, NumInterpolationPoints);
+							break;
+						case EMeshResizingWarpMethod::WrapDeform:
+						default:
+							UE::MeshResizing::Private::ApplyRefitToMesh_WrapDeform(ResizeMesh, *InSourceMesh, BlendedTargetFMesh, 1.0f);
+							break;
+						}
+
+						OutBlendedTargetMesh->SetMesh(MoveTemp(BlendedTargetFMesh));
+						OutResizedMesh->SetMesh(MoveTemp(ResizeMesh));
+
+						if (IsConnected(&MeshToResizeMaterialArray))
+						{
+							OutMaterials = GetValue(Context, &MeshToResizeMaterialArray);
+						}
+					}
 				}
 			}
 		}
diff --git a/Engine/Plugins/Experimental/MeshResizing/Source/MeshResizingNodes/Private/MeshResizing/MeshWarpNode.h b/Engine/Plugins/Experimental/MeshResizing/Source/MeshResizingNodes/Private/MeshResizing/MeshWarpNode.h
--- a/Engine/Plugins/Experimental/MeshResizing/Source/MeshResizingNodes/Private/MeshResizing/MeshWarpNode.h
+++ b/Engine/Plugins/Experimental/MeshResizing/Source/MeshResizingNodes/Private/MeshResizing/MeshWarpNode.h
@@ -10,6 +10,13 @@
 class UDynamicMesh;
 class UMaterialInterface;
 
+UENUM()
+enum struct EMeshResizingWarpMethod : uint8
+{
+	WrapDeform,
+	RBFInterpolate
+};
+
 USTRUCT(Meta = (MeshResizing, Experimental))
 struct FMeshWarpNode : public FDataflowNode
 {
@@ -48,10 +55,17 @@ private:
 	UPROPERTY(meta = (DataflowOutput))
 	TArray<TObjectPtr<UMaterialInterface>> MaterialArray;
 
-
 	UPROPERTY(EditAnywhere, Category = "Warp", meta = (Min = 0.0f, UIMin = 0.0f, Max = 1.0f, UIMax = 1.0f))
 	float Alpha;
 
+	UPROPERTY(EditAnywhere, Category = "Warp")
+	EMeshResizingWarpMethod WarpMethod = EMeshResizingWarpMethod::WrapDeform;
+
+	UPROPERTY(EditAnywhere, Category = "Warp", meta =  (Min = 1, EditCondition = "WarpMethod == EMeshResizingWarpMethod::RBFInterpolate", EditConditionHides))
+	int32 NumInterpolationPoints = 100;
+
+
+
 	//~ Begin FDataflowNode interface
 	virtual void Evaluate(UE::Dataflow::FContext& Context, const FDataflowOutput* Out) const override;
 	//~ End FDataflowNode interface
diff --git a/Engine/Plugins/Experimental/MeshResizing/Source/MeshResizingNodes/Public/MeshResizing/BaseBodyDataflowNodes.h b/Engine/Plugins/Experimental/MeshResizing/Source/MeshResizingNodes/Public/MeshResizing/BaseBodyDataflowNodes.h
deleted file mode 100644
diff --git a/Engine/Plugins/Runtime/HairStrands/Source/HairStrandsCore/Private/GroomBindingBuilder.cpp b/Engine/Plugins/Runtime/HairStrands/Source/HairStrandsCore/Private/GroomBindingBuilder.cpp
--- a/Engine/Plugins/Runtime/HairStrands/Source/HairStrandsCore/Private/GroomBindingBuilder.cpp
+++ b/Engine/Plugins/Runtime/HairStrands/Source/HairStrandsCore/Private/GroomBindingBuilder.cpp
@@ -812,22 +812,6 @@ namespace GroomBinding_RBFWeighting
 
 namespace GroomBinding_RBFWeighting
 {
-	struct FWeightsBuilder
-	{
-		FWeightsBuilder(const uint32 NumRows, const uint32 NumColumns,
-			const FVector3f* SourcePositions, const FVector3f* TargetPositions);
-
-		using EigenMatrix = Eigen::Map<Eigen::Matrix<float, Eigen::Dynamic, Eigen::Dynamic, Eigen::RowMajor>>;
-
-		/** Compute the weights by inverting the matrix*/
-		void ComputeWeights(const uint32 NumRows, const uint32 NumColumns);
-
-		/** Entries in the dense structure */
-		TArray<float> MatrixEntries;
-
-		/** Entries of the matrix inverse */
-		TArray<float> InverseEntries;
-	};
 
 	FWeightsBuilder::FWeightsBuilder(const uint32 NumRows, const uint32 NumColumns,
 		const FVector3f* SourcePositions, const FVector3f* TargetPositions)
@@ -899,6 +883,7 @@ namespace GroomBinding_RBFWeighting
 
 	void FWeightsBuilder::ComputeWeights(const uint32 NumRows, const uint32 NumColumns)
 	{
+		using EigenMatrix = Eigen::Map<Eigen::Matrix<float, Eigen::Dynamic, Eigen::Dynamic, Eigen::RowMajor>>;
 		EigenMatrix WeightsMatrix(MatrixEntries.GetData(), NumRows, NumColumns);
 		EigenMatrix WeightsInverse(InverseEntries.GetData(), NumColumns, NumRows);
 
diff --git a/Engine/Plugins/Runtime/HairStrands/Source/HairStrandsCore/Public/GroomBindingBuilder.h b/Engine/Plugins/Runtime/HairStrands/Source/HairStrandsCore/Public/GroomBindingBuilder.h
--- a/Engine/Plugins/Runtime/HairStrands/Source/HairStrandsCore/Public/GroomBindingBuilder.h
+++ b/Engine/Plugins/Runtime/HairStrands/Source/HairStrandsCore/Public/GroomBindingBuilder.h
@@ -84,4 +84,20 @@ namespace GroomBinding_RBFWeighting
 		/** List of sampled positions */
 		TArray<FVector3f> SamplePositions;
 	};
+
+
+	struct HAIRSTRANDSCORE_API FWeightsBuilder
+	{
+		FWeightsBuilder(const uint32 NumRows, const uint32 NumColumns,
+			const FVector3f* SourcePositions, const FVector3f* TargetPositions);
+
+		/** Compute the weights by inverting the matrix*/
+		void ComputeWeights(const uint32 NumRows, const uint32 NumColumns);
+
+		/** Entries in the dense structure */
+		TArray<float> MatrixEntries;
+
+		/** Entries of the matrix inverse */
+		TArray<float> InverseEntries;
+	};
 }
\ No newline at end of file
