Add a simple framework for mesh resizing iterative deformers.

diff --git a/Engine/Plugins/ChaosClothAssetEditor/Source/ChaosClothAssetDataflowNodes/Private/ChaosClothAsset/ChaosClothAssetDataflowNodesModule.cpp b/Engine/Plugins/ChaosClothAssetEditor/Source/ChaosClothAssetDataflowNodes/Private/ChaosClothAsset/ChaosClothAssetDataflowNodesModule.cpp
--- a/Engine/Plugins/ChaosClothAssetEditor/Source/ChaosClothAssetDataflowNodes/Private/ChaosClothAsset/ChaosClothAssetDataflowNodesModule.cpp
+++ b/Engine/Plugins/ChaosClothAssetEditor/Source/ChaosClothAssetDataflowNodes/Private/ChaosClothAsset/ChaosClothAssetDataflowNodesModule.cpp
@@ -95,6 +95,7 @@ namespace UE::Chaos::ClothAsset
 			DATAFLOW_NODE_REGISTER_CREATION_FACTORY(FChaosClothAssetCopySimulationToRenderMeshNode);
 			DATAFLOW_NODE_REGISTER_CREATION_FACTORY(FChaosClothAssetDatasmithImportNode);
 			DATAFLOW_NODE_REGISTER_CREATION_FACTORY(FChaosClothAssetDeleteElementNode);
+			DATAFLOW_NODE_REGISTER_CREATION_FACTORY(FChaosClothAssetExtractWeightMapNode);
 			DATAFLOW_NODE_REGISTER_CREATION_FACTORY(FChaosClothAssetImportNode);
 			DATAFLOW_NODE_REGISTER_CREATION_FACTORY(FChaosClothAssetImportSimulationCacheNode);
 			DATAFLOW_NODE_REGISTER_CREATION_FACTORY(FChaosClothAssetMergeClothCollectionsNode_v2);
diff --git a/Engine/Plugins/ChaosClothAssetEditor/Source/ChaosClothAssetDataflowNodes/Private/ChaosClothAsset/ClothCollectionToDynamicMeshNode.cpp b/Engine/Plugins/ChaosClothAssetEditor/Source/ChaosClothAssetDataflowNodes/Private/ChaosClothAsset/ClothCollectionToDynamicMeshNode.cpp
--- a/Engine/Plugins/ChaosClothAssetEditor/Source/ChaosClothAssetDataflowNodes/Private/ChaosClothAsset/ClothCollectionToDynamicMeshNode.cpp
+++ b/Engine/Plugins/ChaosClothAssetEditor/Source/ChaosClothAssetDataflowNodes/Private/ChaosClothAsset/ClothCollectionToDynamicMeshNode.cpp
@@ -4,6 +4,7 @@
 #include "ChaosClothAsset/ClothPatternToDynamicMesh.h"
 #include "ChaosClothAsset/ClothPatternToDynamicMeshMappingSupport.h"
 #include "ChaosClothAsset/CollectionClothFacade.h"
+#include "ChaosClothAsset/ClothCollectionGroup.h"
 #include "Dataflow/DataflowInputOutput.h"
 #include "DynamicMesh/DynamicMesh3.h"
 #include "Materials/MaterialInterface.h"
@@ -172,4 +173,71 @@ void FChaosClothAssetUpdateClothFromDynamicMeshNode::Evaluate(UE::Dataflow::FCon
 	}
 }
 
+FChaosClothAssetExtractWeightMapNode::FChaosClothAssetExtractWeightMapNode(const UE::Dataflow::FNodeParameters& InParam, FGuid InGuid)
+	: FDataflowNode(InParam, InGuid)
+{
+	RegisterInputConnection(&Collection);
+	RegisterInputConnection(&WeightMap.StringValue, GET_MEMBER_NAME_CHECKED(FChaosClothAssetConnectableIStringValue, StringValue));
+	RegisterInputConnection(&DynamicMesh)
+		.SetCanHidePin(true)
+		.SetPinIsHidden(true);
+	RegisterOutputConnection(&ExtractedWeightMap);
+}
+
+void FChaosClothAssetExtractWeightMapNode::Evaluate(UE::Dataflow::FContext& Context, const FDataflowOutput* Out) const
+{
+	if (Out->IsA(&ExtractedWeightMap))
+	{
+		using namespace UE::Chaos::ClothAsset;
+
+		TArray<float> Result;
+
+		FManagedArrayCollection InCollection = GetValue<FManagedArrayCollection>(Context, &Collection);
+		const TSharedRef<FManagedArrayCollection> ClothCollection = MakeShared<FManagedArrayCollection>(MoveTemp(InCollection));
+		FCollectionClothConstFacade ClothFacade(ClothCollection);
+		if (ClothFacade.IsValid())
+		{
+			const FName WeightMapName(*GetValue<FString>(Context, &WeightMap.StringValue));
+			switch (MeshTarget)
+			{
+			default:
+			case EChaosClothAssetWeightMapMeshTarget::Simulation:
+				if (ClothFacade.HasWeightMap(WeightMapName))
+				{
+					Result = ClothFacade.GetWeightMap(WeightMapName);
+				}
+				break;
+			case EChaosClothAssetWeightMapMeshTarget::Render:
+				if (ClothFacade.HasUserDefinedAttribute<float>(WeightMapName, ClothCollectionGroup::RenderVertices))
+				{
+					Result = ClothFacade.GetUserDefinedAttribute<float>(WeightMapName, ClothCollectionGroup::RenderVertices);
+				}
+				break;
+			}
+
+			if (bReorderForDynamicMesh && !Result.IsEmpty())
+			{
+				if (const TObjectPtr<const UDynamicMesh> InMesh = GetValue(Context, &DynamicMesh))
+				{
+					const UE::Geometry::FDynamicMesh3& DynMesh = InMesh->GetMeshRef();
+					FClothPatternToDynamicMeshMappingSupport ClothMapping(DynMesh);
+
+					TArray<float> ReorderedResult;
+					ReorderedResult.Init(0.f, DynMesh.MaxVertexID());
+					for (int32 VertexID : DynMesh.VertexIndicesItr())
+					{
+						const int32 ClothMeshID = ClothMapping.GetOriginalVertexID(VertexID);
+						if (Result.IsValidIndex(ClothMeshID))
+						{
+							ReorderedResult[VertexID] = Result[ClothMeshID];
+						}
+					}
+
+					Result = MoveTemp(ReorderedResult);
+				}
+			}
+		}
+		SetValue(Context, MoveTemp(Result), &ExtractedWeightMap);
+	}
+}
 #undef LOCTEXT_NAMESPACE
diff --git a/Engine/Plugins/ChaosClothAssetEditor/Source/ChaosClothAssetDataflowNodes/Public/ChaosClothAsset/ClothCollectionToDynamicMeshNode.h b/Engine/Plugins/ChaosClothAssetEditor/Source/ChaosClothAssetDataflowNodes/Public/ChaosClothAsset/ClothCollectionToDynamicMeshNode.h
--- a/Engine/Plugins/ChaosClothAssetEditor/Source/ChaosClothAssetDataflowNodes/Public/ChaosClothAsset/ClothCollectionToDynamicMeshNode.h
+++ b/Engine/Plugins/ChaosClothAssetEditor/Source/ChaosClothAssetDataflowNodes/Public/ChaosClothAsset/ClothCollectionToDynamicMeshNode.h
@@ -3,6 +3,8 @@
 #pragma once 
 
 #include "Dataflow/DataflowNode.h"
+#include "ChaosClothAsset/ConnectableValue.h"
+#include "ChaosClothAsset/WeightMapNode.h"
 #include "GeometryCollection/ManagedArrayCollection.h"
 #include "ClothCollectionToDynamicMeshNode.generated.h"
 
@@ -87,3 +89,43 @@ private:
 	UPROPERTY(EditAnywhere, Category = "DynamicMesh To Cloth Collection", meta = (EditCondition = "bCopyUVsToSim2DPositions", EditConditionHides))
 	int32 UVChannelIndex = 0;
 };
+
+/** Extract a weight map from a Cloth Collection.
+ */
+USTRUCT(Meta = (DataflowCloth, Experimental))
+struct FChaosClothAssetExtractWeightMapNode : public FDataflowNode
+{
+	GENERATED_USTRUCT_BODY()
+	DATAFLOW_NODE_DEFINE_INTERNAL(FChaosClothAssetExtractWeightMapNode, "ExtractClothWeightMap", "Cloth", "Extract Cloth Weight Map")
+
+public:
+
+	FChaosClothAssetExtractWeightMapNode(const UE::Dataflow::FNodeParameters& InParam, FGuid InGuid = FGuid::NewGuid());
+
+	virtual void Evaluate(UE::Dataflow::FContext& Context, const FDataflowOutput* Out) const override;
+
+private:
+
+	UPROPERTY(Meta = (DataflowInput))
+	FManagedArrayCollection Collection;
+
+	/** Name of the weight map to be extracted */
+	UPROPERTY(EditAnywhere, Category = "Extract Weight Map", Meta = (DataflowInput))
+	FChaosClothAssetConnectableIStringValue WeightMap = { TEXT("WeightMap") };
+
+	/** Type of the weight map to be extracted */
+	UPROPERTY(EditAnywhere, Category = "Extract Weight Map")
+	EChaosClothAssetWeightMapMeshTarget MeshTarget = EChaosClothAssetWeightMapMeshTarget::Simulation;
+
+	/** Reorder extracted weights to match the order of a DynamicMesh that was created via ClothCollectionToDynamicMesh */
+	UPROPERTY(EditAnywhere, Category = "Extract Weight Map")
+	bool bReorderForDynamicMesh = true;
+
+	/** Dynamic mesh used to reorder weights. */
+	UPROPERTY(meta = (DataflowInput, EditCondition = "bReorderForDynamicMesh"))
+	TObjectPtr<UDynamicMesh> DynamicMesh;
+
+	UPROPERTY(meta = (DataflowOutput))
+	TArray<float> ExtractedWeightMap;
+
+};
\ No newline at end of file
diff --git a/Engine/Plugins/Experimental/MeshResizing/Source/MeshResizingCore/Private/MeshResizing/Mesh3DConstraints.cpp b/Engine/Plugins/Experimental/MeshResizing/Source/MeshResizingCore/Private/MeshResizing/Mesh3DConstraints.cpp
new file mode 100644
--- /dev/null
+++ b/Engine/Plugins/Experimental/MeshResizing/Source/MeshResizingCore/Private/MeshResizing/Mesh3DConstraints.cpp
@@ -0,0 +1,122 @@
+// Copyright Epic Games, Inc. All Rights Reserved.
+
+#include "MeshResizing/Mesh3DConstraints.h"
+#include "DynamicMesh/DynamicMesh3.h"
+
+namespace UE::MeshResizing
+{
+	static bool DO_NAN_CHECK = true;
+	static double LENGTH_CHECK = 2.;
+
+	FShearConstraint::FShearConstraint(float ShearConstraintStrength, const TArray<float>& ShearConstraintWeights, const int32 InNumParticles)
+	: NumParticles(InNumParticles)
+	, ShearWeightMap(Chaos::Softs::FSolverVec2(0.f, ShearConstraintStrength), ShearConstraintWeights, NumParticles)
+	{
+		if (!ShearWeightMap.HasWeightMap())
+		{
+			// Treat as constant EdgeRotationStrength rather than constant 0
+			ShearWeightMap.SetWeightedValue(Chaos::Softs::FSolverVec2(ShearConstraintStrength, ShearConstraintStrength));
+		}
+	}
+
+	void FShearConstraint::Apply(UE::Geometry::FDynamicMesh3& ResizedMesh, const UE::Geometry::FDynamicMesh3& InitialResizedMesh, const UE::Geometry::FDynamicMesh3& BaseMesh) const
+	{
+		check(ResizedMesh.MaxVertexID() <= NumParticles);
+		using namespace UE::Geometry;
+		using namespace Chaos::Softs;
+
+		// Do a Jacobi update
+		UE::Geometry::FDynamicMesh3 ResizedMesh0;
+		ResizedMesh0.Copy(ResizedMesh);
+
+		for (int32 TriIndex : ResizedMesh.TriangleIndicesItr())
+		{
+			const FIndex3i& Tri = ResizedMesh.GetTriangleRef(TriIndex);
+			if (!ensure(BaseMesh.IsVertex(Tri[0]) && BaseMesh.IsVertex(Tri[1]) && BaseMesh.IsVertex(Tri[2])))
+			{
+				continue;
+			}
+			const float Shear0 = ShearWeightMap.GetValue(Tri[0]);
+			const float Shear1 = ShearWeightMap.GetValue(Tri[1]);
+			const float Shear2 = ShearWeightMap.GetValue(Tri[2]);
+			if (Shear0 == 0.f && Shear1 == 0.f && Shear2 == 0.f)
+			{
+				continue;
+			}
+
+			// calculate the transformation between the current resized mesh and the base mesh.
+			// Using Chaos notation of P = deforming positions, X = base positions
+			const FVector3d& X0 = BaseMesh.GetVertexRef(Tri[0]);
+			const FVector3d& X1 = BaseMesh.GetVertexRef(Tri[1]);
+			const FVector3d& X2 = BaseMesh.GetVertexRef(Tri[2]);
+
+			const FVector3d X10 = X1 - X0;
+			const FVector3d X20 = X2 - X0;
+			const FVector3d X10xX20 = FVector3d::CrossProduct(X10, X20).GetSafeNormal(UE_SMALL_NUMBER);
+
+			// If the base triangle is degenerate, we won't be able to find the transformation.
+			if (X10xX20.IsNearlyZero())
+			{
+				continue;
+			}
+
+			const FVector3d P0 = ResizedMesh0.GetVertex(Tri[0]);
+			const FVector3d P1 = ResizedMesh0.GetVertex(Tri[1]);
+			const FVector3d P2 = ResizedMesh0.GetVertex(Tri[2]);
+			const FVector3d P10 = P1 - P0;
+			const FVector3d P20 = P2 - P0;
+			const FVector3d P10xP20 = FVector3d::CrossProduct(P10, P20).GetSafeNormal(UE_SMALL_NUMBER);
+			if (P10xP20.IsNearlyZero())
+			{
+				continue;
+			}
+			const FMatrix3d PMat(P10, P20, P10xP20, false);
+			const FMatrix3d XMat(X10, X20, X10xX20, false);
+			const FMatrix3d XMatInv = XMat.Inverse();
+			const FMatrix3d TransformMat = PMat * XMatInv;
+			const FMatrix3d TransformMatT = TransformMat.Transpose();
+
+			TMatrix<double> TransformMat44(TransformMatT.Row0, TransformMatT.Row1, TransformMatT.Row2, FVector3d::ZeroVector);
+
+			check(TransformMat44.TransformPosition(X10).Equals(P10));
+			check(TransformMat44.TransformPosition(X20).Equals(P20));
+			check(TransformMat44.TransformPosition(X10xX20).Equals(P10xP20));
+
+			TVector<double> Scale = TransformMat44.ExtractScaling();
+			if (TransformMat44.Determinant() < 0.f)
+			{
+				// Assume it is along X and modify transform accordingly. 
+				// It doesn't actually matter which axis we choose, the 'appearance' will be the same
+				Scale.X *= -1.f;
+				TransformMat44.SetAxis(0, -TransformMat44.GetScaledAxis(EAxis::X));
+			}
+
+			TQuat<double> InRotation = TQuat<double>(TransformMat44);
+
+
+			const FVector3d PInit0 = InitialResizedMesh.GetVertex(Tri[0]);
+			const FVector3d PInit1 = InitialResizedMesh.GetVertex(Tri[1]);
+			const FVector3d PInit2 = InitialResizedMesh.GetVertex(Tri[2]);
+
+			FVector3d P10New = X10.GetSafeNormal() * (PInit1 - PInit0).Length();// *Scale;
+			P10New = InRotation.RotateVector(P10New);
+			check(!DO_NAN_CHECK || !P10New.ContainsNaN());
+
+			FVector3d P20New = X20.GetSafeNormal() * (PInit2 - PInit0).Length();// *Scale;
+			P20New = InRotation.RotateVector(P20New);
+			check(!DO_NAN_CHECK || !P20New.ContainsNaN());
+			const FVector3d Delta20 = P20New - P20;
+			const FVector3d Delta10 = P10New - P10;
+
+			//check(!DO_NAN_CHECK || Delta20.SquaredLength() < LENGTH_CHECK);
+			//check(!DO_NAN_CHECK || Delta10.SquaredLength() < LENGTH_CHECK);
+			const FVector3d P1New = ResizedMesh.GetVertex(Tri[1]) + Shear1 * Delta10;
+			const FVector3d P2New = ResizedMesh.GetVertex(Tri[2]) + Shear2 * Delta20;
+			const FVector3d P0New = ResizedMesh.GetVertex(Tri[0]) - Shear0 * (Delta10 + Delta20);
+
+			ResizedMesh.SetVertex(Tri[0], P0New);
+			ResizedMesh.SetVertex(Tri[1], P1New);
+			ResizedMesh.SetVertex(Tri[2], P2New);
+		}
+	}
+}
\ No newline at end of file
diff --git a/Engine/Plugins/Experimental/MeshResizing/Source/MeshResizingCore/Public/MeshResizing/Mesh3DConstraints.h b/Engine/Plugins/Experimental/MeshResizing/Source/MeshResizingCore/Public/MeshResizing/Mesh3DConstraints.h
new file mode 100644
--- /dev/null
+++ b/Engine/Plugins/Experimental/MeshResizing/Source/MeshResizingCore/Public/MeshResizing/Mesh3DConstraints.h
@@ -0,0 +1,25 @@
+// Copyright Epic Games, Inc. All Rights Reserved.
+
+#pragma once 
+#include "UObject/NameTypes.h"
+#include "Chaos/PBDFlatWeightMap.h"
+#include "Containers/Array.h"
+
+namespace UE::Geometry
+{
+	class FDynamicMesh3;
+}
+namespace UE::MeshResizing
+{
+	class FShearConstraint
+	{
+	public:
+		MESHRESIZINGCORE_API FShearConstraint(float ShearConstraintStrength, const TArray<float>& ShearConstraintWeights, const int32 NumParticles);
+		MESHRESIZINGCORE_API void Apply(UE::Geometry::FDynamicMesh3& ResizedMesh, const UE::Geometry::FDynamicMesh3& InitialResizedMesh, const UE::Geometry::FDynamicMesh3& BaseMesh) const;
+
+	private:
+		int32 NumParticles;
+		Chaos::Softs::FPBDFlatWeightMapView ShearWeightMap;
+	};
+}
+
diff --git a/Engine/Plugins/Experimental/MeshResizing/Source/MeshResizingNodes/Private/MeshResizing/MeshConstraintNodes.cpp b/Engine/Plugins/Experimental/MeshResizing/Source/MeshResizingNodes/Private/MeshResizing/MeshConstraintNodes.cpp
new file mode 100644
--- /dev/null
+++ b/Engine/Plugins/Experimental/MeshResizing/Source/MeshResizingNodes/Private/MeshResizing/MeshConstraintNodes.cpp
@@ -0,0 +1,95 @@
+// Copyright Epic Games, Inc. All Rights Reserved.
+#include "MeshResizing/MeshConstraintNodes.h"
+#include "Dataflow/DataflowNodeFactory.h"
+#include "Dataflow/DataflowMesh.h"
+#include "MeshResizing/Mesh3DConstraints.h"
+#include "UDynamicMesh.h"
+
+
+#include UE_INLINE_GENERATED_CPP_BY_NAME(MeshConstraintNodes)
+
+using UE::Geometry::FDynamicMesh3;
+
+namespace  UE::MeshResizing
+{
+	void RegisterMeshConstraintDataflowNodes()
+	{
+		static const FLinearColor CDefaultNodeBodyTintColor = FLinearColor(0.f, 0.f, 0.f, 0.5f);
+
+		DATAFLOW_NODE_REGISTER_CREATION_FACTORY(FMeshConstrainedDeformationNode);
+	}
+
+	namespace Private
+	{
+		typedef TFunction<void (UE::Geometry::FDynamicMesh3& ResizedMesh, const UE::Geometry::FDynamicMesh3& InitialResizedMesh, const UE::Geometry::FDynamicMesh3& BaseMesh, int32 Iteration)> ApplyConstraintFunc;
+
+		static void Solve(UE::Geometry::FDynamicMesh3& ResizedMesh, const UE::Geometry::FDynamicMesh3& InitialResizedMesh, const UE::Geometry::FDynamicMesh3& BaseMesh, int32 Iterations,
+			const TArray<ApplyConstraintFunc>& Constraints)
+		{
+			if (Constraints.IsEmpty())
+			{
+				return;
+			}
+			for (int32 Iteration = 0; Iteration < Iterations; ++Iteration)
+			{
+				for (const ApplyConstraintFunc& Constraint : Constraints)
+				{
+					Constraint(ResizedMesh, InitialResizedMesh, BaseMesh, Iteration);
+				}
+			}
+		}
+	}
+}
+
+FMeshConstrainedDeformationNode::FMeshConstrainedDeformationNode(const UE::Dataflow::FNodeParameters& InParam, FGuid InGuid)
+	: FDataflowNode(InParam, InGuid)
+{
+	RegisterInputConnection(&ResizingMesh);
+	RegisterInputConnection(&BaseMesh);
+	RegisterInputConnection(&ShearConstraintWeights);
+	RegisterOutputConnection(&ResizingMesh, &ResizingMesh);
+}
+
+void FMeshConstrainedDeformationNode::Evaluate(UE::Dataflow::FContext& Context, const FDataflowOutput* Out) const
+{
+	using namespace UE::Geometry;
+	using namespace UE::MeshResizing;
+	if (Out->IsA(&ResizingMesh))
+	{
+		if (TObjectPtr<UDataflowMesh> InResizingMesh = GetValue(Context, &ResizingMesh))
+		{
+			if (TObjectPtr<UDataflowMesh> InBaseMesh = GetValue(Context, &BaseMesh))
+			{
+				if (InResizingMesh->GetDynamicMesh() && InBaseMesh->GetDynamicMesh())
+				{
+					FDynamicMesh3 ResizedMesh;
+					ResizedMesh.Copy(InResizingMesh->GetDynamicMeshRef());
+
+					TArray<Private::ApplyConstraintFunc> Constraints;
+
+
+					UE::MeshResizing::FShearConstraint ShearConstraint(ShearConstraintStrength, GetValue(Context, &ShearConstraintWeights), ResizedMesh.MaxVertexID());
+					if (bEnableShearConstraint)
+					{
+						Constraints.Add(
+							[&ShearConstraint](UE::Geometry::FDynamicMesh3& ResizedMesh, const UE::Geometry::FDynamicMesh3& InitialResizedMesh, const UE::Geometry::FDynamicMesh3& BaseMesh, int32 Iteration)
+							{
+								ShearConstraint.Apply(ResizedMesh, InitialResizedMesh, BaseMesh);
+							}
+						);
+					}
+
+					UE::MeshResizing::Private::Solve(ResizedMesh, InResizingMesh->GetDynamicMeshRef(), InBaseMesh->GetDynamicMeshRef(), Iterations, Constraints);
+
+					TObjectPtr<UDataflowMesh> OutResizedMesh = NewObject<UDataflowMesh>();
+					OutResizedMesh->SetDynamicMesh(MoveTemp(ResizedMesh));
+					OutResizedMesh->SetMaterials(InResizingMesh->GetMaterials());
+					SetValue(Context, OutResizedMesh, &ResizingMesh);
+					return;
+				}
+			}
+		}
+		SafeForwardInput(Context, &ResizingMesh, &ResizingMesh);
+	}
+}
+
diff --git a/Engine/Plugins/Experimental/MeshResizing/Source/MeshResizingNodes/Private/MeshResizing/MeshConstraintNodes.h b/Engine/Plugins/Experimental/MeshResizing/Source/MeshResizingNodes/Private/MeshResizing/MeshConstraintNodes.h
new file mode 100644
--- /dev/null
+++ b/Engine/Plugins/Experimental/MeshResizing/Source/MeshResizingNodes/Private/MeshResizing/MeshConstraintNodes.h
@@ -0,0 +1,54 @@
+// Copyright Epic Games, Inc. All Rights Reserved.
+
+#pragma once
+
+#include "Dataflow/DataflowNode.h"
+#include "GeometryCollection/ManagedArrayCollection.h"
+
+#include "MeshConstraintNodes.generated.h"
+
+class UDataflowMesh;
+class UMaterialInterface;
+
+
+USTRUCT(Meta = (MeshResizing, Experimental))
+struct FMeshConstrainedDeformationNode : public FDataflowNode
+{
+	GENERATED_USTRUCT_BODY()
+	DATAFLOW_NODE_DEFINE_INTERNAL(FMeshConstrainedDeformationNode, "MeshConstrainedDeformation", "MeshResizing", "Mesh Constrained Deformation")
+	DATAFLOW_NODE_RENDER_TYPE("SurfaceRender", FName("UDataflowMesh"), "ResizingMesh")
+
+public:
+
+	FMeshConstrainedDeformationNode(const UE::Dataflow::FNodeParameters& InParam, FGuid InGuid = FGuid::NewGuid());
+
+private:
+
+	UPROPERTY(meta = (DataflowInput, DataflowOutput, DataflowPassthrough = "ResizingMesh"))
+	TObjectPtr<UDataflowMesh> ResizingMesh;
+
+	UPROPERTY(meta = (DataflowInput))
+	TObjectPtr<UDataflowMesh> BaseMesh;
+
+	UPROPERTY(EditAnywhere, Category = "Constrained Deformation|Solver", meta = (Min = 0, UIMax = 100))
+	int32 Iterations = 100;
+
+	/** Remove shear deformation */
+	UPROPERTY(EditAnywhere, Category = "Constrained Deformation|Shear Constraint")
+	bool bEnableShearConstraint = true;
+
+	UPROPERTY(EditAnywhere, Category = "Constrained Deformation|Shear Constraint", meta = (Min = 0, Max = 1, EditCondition = "bEnableShearConstraint"))
+	float ShearConstraintStrength = 1.f;
+
+	UPROPERTY(meta = (DataflowInput))
+	TArray<float> ShearConstraintWeights;
+
+	//~ Begin FDataflowNode interface
+	virtual void Evaluate(UE::Dataflow::FContext& Context, const FDataflowOutput* Out) const override;
+	//~ End FDataflowNode interface
+};
+
+namespace UE::MeshResizing
+{
+	void RegisterMeshConstraintDataflowNodes();
+}
\ No newline at end of file
diff --git a/Engine/Plugins/Experimental/MeshResizing/Source/MeshResizingNodes/Private/MeshResizing/MeshResizingDataflowNodesModule.cpp b/Engine/Plugins/Experimental/MeshResizing/Source/MeshResizingNodes/Private/MeshResizing/MeshResizingDataflowNodesModule.cpp
--- a/Engine/Plugins/Experimental/MeshResizing/Source/MeshResizingNodes/Private/MeshResizing/MeshResizingDataflowNodesModule.cpp
+++ b/Engine/Plugins/Experimental/MeshResizing/Source/MeshResizingNodes/Private/MeshResizing/MeshResizingDataflowNodesModule.cpp
@@ -4,6 +4,7 @@
 #include "Dataflow/DataflowNodeFactory.h"
 #include "Modules/ModuleManager.h"
 #include "MeshResizing/BaseBodyDataflowNodes.h"
+#include "MeshResizing/MeshConstraintNodes.h"
 #include "MeshResizing/MeshWarpNode.h"
 
 namespace UE::MeshResizing
@@ -22,6 +23,7 @@ namespace UE::MeshResizing
 			DATAFLOW_NODE_REGISTER_CREATION_FACTORY_NODE_COLORS_BY_CATEGORY("MeshResizing", FColorScheme::NodeHeader, FColorScheme::NodeBody);
 
 			RegisterBaseBodyDataflowNodes();
+			RegisterMeshConstraintDataflowNodes();
 
 			DATAFLOW_NODE_REGISTER_CREATION_FACTORY(FMeshWarpNode);
 		}
