Fix CIS error from previous submit in MeshResizing.

diff --git a/Engine/Plugins/Experimental/MeshResizing/Source/MeshResizingNodes/Private/MeshResizing/MeshConstraintNodes.cpp b/Engine/Plugins/Experimental/MeshResizing/Source/MeshResizingNodes/Private/MeshResizing/MeshConstraintNodes.cpp
--- a/Engine/Plugins/Experimental/MeshResizing/Source/MeshResizingNodes/Private/MeshResizing/MeshConstraintNodes.cpp
+++ b/Engine/Plugins/Experimental/MeshResizing/Source/MeshResizingNodes/Private/MeshResizing/MeshConstraintNodes.cpp
@@ -72,9 +72,9 @@ void FMeshConstrainedDeformationNode::Evaluate(UE::Dataflow::FContext& Context,
 					if (bEnableShearConstraint)
 					{
 						Constraints.Add(
-							[&ShearConstraint](UE::Geometry::FDynamicMesh3& ResizedMesh, const UE::Geometry::FDynamicMesh3& InitialResizedMesh, const UE::Geometry::FDynamicMesh3& BaseMesh, int32 Iteration)
+							[&ShearConstraint](UE::Geometry::FDynamicMesh3& InResizedMesh, const UE::Geometry::FDynamicMesh3& InInitialResizedMesh, const UE::Geometry::FDynamicMesh3& InBaseMesh, int32 Iteration)
 							{
-								ShearConstraint.Apply(ResizedMesh, InitialResizedMesh, BaseMesh);
+								ShearConstraint.Apply(InResizedMesh, InInitialResizedMesh, InBaseMesh);
 							}
 						);
 					}
