Make MeshWarp node passthrough on failure/disable and other optimizations and default changes.

Set default Alpha = 1.f
Switch default Interpolation Method to RBF since that's the one we seem to be landing on.
Don't calculate RBF weights if the resizing mesh is empty.

diff --git a/Engine/Plugins/Experimental/MeshResizing/Source/MeshResizingNodes/Private/MeshResizing/MeshWarpNode.cpp b/Engine/Plugins/Experimental/MeshResizing/Source/MeshResizingNodes/Private/MeshResizing/MeshWarpNode.cpp
--- a/Engine/Plugins/Experimental/MeshResizing/Source/MeshResizingNodes/Private/MeshResizing/MeshWarpNode.cpp
+++ b/Engine/Plugins/Experimental/MeshResizing/Source/MeshResizingNodes/Private/MeshResizing/MeshWarpNode.cpp
@@ -100,6 +100,11 @@ namespace  UE::MeshResizing::Private
 		const FDynamicMesh3& TargetMesh,
 		int32 NumInterpolationPoints)
 	{
+		if (ResizedMesh.VertexCount() == 0)
+		{
+			return;
+		}
+
 		FRBFInterpolation::FInterpolationData InterpolationData;
 		FRBFInterpolation::GenerateWeights(SourceMesh, NumInterpolationPoints, InterpolationData);
 
@@ -123,16 +128,14 @@ FMeshWarpNode::FMeshWarpNode(const UE::Dataflow::FNodeParameters& InParam, FGuid
 	RegisterInputConnection(&SourceMesh);
 	RegisterInputConnection(&TargetMesh);
 
-	RegisterOutputConnection(&BlendedTargetMesh);
-	RegisterOutputConnection(&ResizedMesh);
+	RegisterOutputConnection(&BlendedTargetMesh, &SourceMesh);
+	RegisterOutputConnection(&ResizedMesh, &MeshToResize);
 }
 
 void FMeshWarpNode::Evaluate(UE::Dataflow::FContext& Context, const FDataflowOutput* Out) const
 {
 	using namespace UE::Geometry;
 
-	TObjectPtr<UDataflowMesh> OutResizedMesh = NewObject<UDataflowMesh>();
-	TObjectPtr<UDataflowMesh> OutBlendedTargetMesh = NewObject<UDataflowMesh>();
 
 	if (TObjectPtr<UDataflowMesh> InMeshToResize = GetValue(Context, &MeshToResize))
 	{
@@ -142,6 +145,8 @@ void FMeshWarpNode::Evaluate(UE::Dataflow::FContext& Context, const FDataflowOut
 			{
 				if (InMeshToResize->GetDynamicMesh() && InSourceMesh->GetDynamicMesh() && InTargetMesh->GetDynamicMesh())
 				{
+					TObjectPtr<UDataflowMesh> OutResizedMesh = NewObject<UDataflowMesh>();
+					TObjectPtr<UDataflowMesh> OutBlendedTargetMesh = NewObject<UDataflowMesh>();
 					FDynamicMesh3 BlendedTargetFMesh;
 					UE::MeshResizing::Private::BlendMesh(BlendedTargetFMesh, InSourceMesh->GetDynamicMeshRef(), InTargetMesh->GetDynamicMeshRef(), Alpha);
 
@@ -150,25 +155,27 @@ void FMeshWarpNode::Evaluate(UE::Dataflow::FContext& Context, const FDataflowOut
 					switch (WarpMethod)
 					{
 					case EMeshResizingWarpMethod::RBFInterpolate:
+					default:
 						UE::MeshResizing::Private::ApplyRefitToMesh_RBFInterpolate(ResizeMesh, InSourceMesh->GetDynamicMeshRef(), BlendedTargetFMesh, NumInterpolationPoints);
 						break;
 					case EMeshResizingWarpMethod::WrapDeform:
-					default:
 						UE::MeshResizing::Private::ApplyRefitToMesh_WrapDeform(ResizeMesh, InSourceMesh->GetDynamicMeshRef(), BlendedTargetFMesh, 1.0f);
 						break;
 					}
 
 					OutBlendedTargetMesh->SetDynamicMesh(MoveTemp(BlendedTargetFMesh));
 					OutResizedMesh->SetDynamicMesh(MoveTemp(ResizeMesh));
+					OutBlendedTargetMesh->SetMaterials(InTargetMesh->GetMaterials());
+					OutResizedMesh->SetMaterials(InMeshToResize->GetMaterials());
+					SetValue(Context, OutBlendedTargetMesh, &BlendedTargetMesh);
+					SetValue(Context, OutResizedMesh, &ResizedMesh);
+					return;
 				}
-
-				OutBlendedTargetMesh->SetMaterials(InTargetMesh->GetMaterials());
-				OutResizedMesh->SetMaterials(InMeshToResize->GetMaterials());
 			}
 		}
 	}
 
-	SetValue(Context, MoveTemp(OutResizedMesh), &ResizedMesh);
-	SetValue(Context, MoveTemp(OutBlendedTargetMesh), &BlendedTargetMesh);
+	SafeForwardInput(Context, &MeshToResize, &ResizedMesh);
+	SafeForwardInput(Context, &SourceMesh, &BlendedTargetMesh);
 }
 
diff --git a/Engine/Plugins/Experimental/MeshResizing/Source/MeshResizingNodes/Private/MeshResizing/MeshWarpNode.h b/Engine/Plugins/Experimental/MeshResizing/Source/MeshResizingNodes/Private/MeshResizing/MeshWarpNode.h
--- a/Engine/Plugins/Experimental/MeshResizing/Source/MeshResizingNodes/Private/MeshResizing/MeshWarpNode.h
+++ b/Engine/Plugins/Experimental/MeshResizing/Source/MeshResizingNodes/Private/MeshResizing/MeshWarpNode.h
@@ -43,17 +43,17 @@ private:
 	TObjectPtr<UDataflowMesh> TargetMesh;
 
 
-	UPROPERTY(meta = (DataflowOutput))
+	UPROPERTY(meta = (DataflowOutput, DataflowPassthrough = "SourceMesh"))
 	TObjectPtr<UDataflowMesh> BlendedTargetMesh;
 
-	UPROPERTY(meta = (DataflowOutput))
+	UPROPERTY(meta = (DataflowOutput, DataflowPassthrough = "MeshToResize"))
 	TObjectPtr<UDataflowMesh> ResizedMesh;
 
 	UPROPERTY(EditAnywhere, Category = "Warp", meta = (Min = 0.0f, UIMin = 0.0f, Max = 1.0f, UIMax = 1.0f))
-	float Alpha;
+	float Alpha = 1.f;
 
 	UPROPERTY(EditAnywhere, Category = "Warp")
-	EMeshResizingWarpMethod WarpMethod = EMeshResizingWarpMethod::WrapDeform;
+	EMeshResizingWarpMethod WarpMethod = EMeshResizingWarpMethod::RBFInterpolate;
 
 	UPROPERTY(EditAnywhere, Category = "Warp", meta =  (Min = 1, EditCondition = "WarpMethod == EMeshResizingWarpMethod::RBFInterpolate", EditConditionHides))
 	int32 NumInterpolationPoints = 100;
