Dataflow: UV Mesh transform and best-fit rigid transform nodes

diff --git a/Engine/Plugins/Experimental/Dataflow/Source/DataflowEditor/Private/Dataflow/DataflowEngineRendering.cpp b/Engine/Plugins/Experimental/Dataflow/Source/DataflowEditor/Private/Dataflow/DataflowEngineRendering.cpp
--- a/Engine/Plugins/Experimental/Dataflow/Source/DataflowEditor/Private/Dataflow/DataflowEngineRendering.cpp
+++ b/Engine/Plugins/Experimental/Dataflow/Source/DataflowEditor/Private/Dataflow/DataflowEngineRendering.cpp
@@ -183,7 +183,8 @@ namespace UE::Dataflow
 
 		void RenderDynamicMeshUV(GeometryCollection::Facades::FRenderingFacade& RenderCollection,
 			const FString& GeometryGroupName,
-			const UE::Geometry::FDynamicMesh3& DynamicMesh)
+			const UE::Geometry::FDynamicMesh3& DynamicMesh,
+			int32 UVChannel)
 		{
 			using namespace UE::Geometry;
 
@@ -191,9 +192,9 @@ namespace UE::Dataflow
 			const int32 NumTriangles = DynamicMesh.TriangleCount();
 			const int32 NumUVLayers = DynamicMesh.Attributes()->NumUVLayers();
 
-			if (NumVertices > 0 && NumTriangles > 0 && NumUVLayers > 0)
+			if (NumVertices > 0 && NumTriangles > 0 && NumUVLayers > UVChannel)
 			{
-				const FDynamicMeshUVOverlay* const UVOverlay = DynamicMesh.Attributes()->PrimaryUV();
+				const FDynamicMeshUVOverlay* const UVOverlay = DynamicMesh.Attributes()->GetUVLayer(UVChannel);
 
 				TArray<FIntVector> Tris; 
 				Tris.Reserve(DynamicMesh.TriangleCount());
@@ -729,6 +730,11 @@ namespace UE::Dataflow
 		{
 			if (State.GetRenderOutputs().Num())
 			{
+				const TArray<FName>& RenderOutputs = State.GetRenderOutputs();
+				if (RenderOutputs.Num() == 0)
+				{
+					return; // no outputs
+				}
 				FName PrimaryOutput = State.GetRenderOutputs()[0];
 
 				TObjectPtr<UDataflowMesh> Default;
@@ -742,7 +748,17 @@ namespace UE::Dataflow
 						}
 						else if (State.GetViewMode().GetName() == Dataflow::FDataflowConstructionUVViewMode::Name)
 						{
-							Private::RenderDynamicMeshUV(RenderCollection, State.GetGuid().ToString(), *DynamicMesh);
+							if (DynamicMesh->HasAttributes())
+							{
+								const FName UVChannelOutput = (RenderOutputs.Num() > 1) ? RenderOutputs[1] : NAME_None; // "UVChannel"
+								const int32 DefaultUVChannel = 0;
+								const int32 UVChannel = (UVChannelOutput.IsNone()) ? DefaultUVChannel : State.GetValue<int32>(UVChannelOutput, DefaultUVChannel);
+
+								if (DynamicMesh->Attributes()->NumUVLayers() > UVChannel)
+								{
+									Private::RenderDynamicMeshUV(RenderCollection, State.GetGuid().ToString(), *DynamicMesh, UVChannel);
+								}
+							}
 						}
 					}
 				}
diff --git a/Engine/Plugins/Experimental/MeshResizing/Source/MeshResizingNodes/Private/MeshResizing/AlignUVMeshNode.cpp b/Engine/Plugins/Experimental/MeshResizing/Source/MeshResizingNodes/Private/MeshResizing/AlignUVMeshNode.cpp
new file mode 100644
--- /dev/null
+++ b/Engine/Plugins/Experimental/MeshResizing/Source/MeshResizingNodes/Private/MeshResizing/AlignUVMeshNode.cpp
@@ -0,0 +1,197 @@
+// Copyright Epic Games, Inc. All Rights Reserved.
+#include "MeshResizing/AlignUVMeshNode.h"
+#include "Dataflow/DataflowNodeFactory.h"
+#include "Dataflow/DataflowMesh.h"
+#include "UDynamicMesh.h"
+#include "Math/TransformCalculus2D.h"
+#include "Selections/MeshConnectedComponents.h"
+
+#include UE_INLINE_GENERATED_CPP_BY_NAME(AlignUVMeshNode)
+
+using UE::Geometry::FDynamicMesh3;
+
+namespace UE::MeshResizing
+{
+	void RegisterAlignUVMeshNodes()
+	{
+		DATAFLOW_NODE_REGISTER_CREATION_FACTORY(FAlignUVMeshNode);
+	}
+
+	namespace Private
+	{
+		// Attempt to find a single rotation, translation and uniform scale that best transforms Source to Dest (AKA Procrustes/Procrustean transform)
+		FTransform2f BestFitTransform(const TArray<FVector2f>& Source, const TArray<FVector2f>& Dest, bool bAllowScale)
+		{	
+			// minimize_T |S*T - D|^2,  where S = Source, D = Dest, T is a (rotation, scale, translation) linear transform
+			// Solve A^T * A * x = A^T * b,   where x = {s*cos(theta), s*sin(theta), tx, ty}
+
+			FMatrix44d ATA(EForceInit::ForceInitToZero);		// A^T * A
+			for (const FVector2f& S : Source)
+			{
+				ATA.M[0][0] += S[0] * S[0] + S[1] * S[1];
+				ATA.M[0][2] += S[0];
+				ATA.M[2][0] += S[0];
+				ATA.M[0][3] += S[1];
+				ATA.M[3][0] += S[1];
+				ATA.M[1][1] += S[0] * S[0] + S[1] * S[1];
+				ATA.M[1][2] += -S[1];
+				ATA.M[2][1] += -S[1];
+				ATA.M[1][3] += S[0];
+				ATA.M[3][1] += S[0];
+				ATA.M[2][2] += 1;
+				ATA.M[3][3] += 1;
+			}
+
+			FVector4d ATb(0, 0, 0, 0);			// A^T * b
+			for (int32 VectorIndex = 0; VectorIndex < Source.Num(); ++VectorIndex)
+			{
+				const FVector2f& S = Source[VectorIndex];
+				const FVector2f& D = Dest[VectorIndex];
+
+				ATb[0] += D[0] * S[0] + D[1] * S[1];
+				ATb[1] += -D[0] * S[1] + D[1] * S[0];
+				ATb[2] += D[0];
+				ATb[3] += D[1];
+			}
+
+			const FMatrix44d ATAInv = ATA.Inverse();
+			const FVector4d Solution = ATAInv.TransformFVector4(ATb);
+
+			float CosTheta = Solution[0];		// Actually s*cos(theta) for uniform scale s
+			float SinTheta = Solution[1];
+
+			if (!bAllowScale)
+			{
+				// remove the scaling
+				const float Scale = FMath::Sqrt(CosTheta * CosTheta + SinTheta * SinTheta);
+				CosTheta = Solution[0] / Scale;
+				SinTheta = Solution[1] / Scale;
+			}
+			
+			const FVector2f Translation(Solution[2], Solution[3]);
+			const FMatrix2x2f ScaleAndRotation(CosTheta, SinTheta, -SinTheta, CosTheta);
+			return FTransform2f(ScaleAndRotation, Translation);
+		}
+	}
+}
+
+FAlignUVMeshNode::FAlignUVMeshNode(const UE::Dataflow::FNodeParameters& InParam, FGuid InGuid)
+	: FDataflowNode(InParam, InGuid)
+{
+	RegisterInputConnection(&ResizingMesh);
+	RegisterInputConnection(&BaseMesh);
+	RegisterInputConnection(&UVChannelIndex);
+	RegisterOutputConnection(&ResizingMesh, &ResizingMesh);
+	RegisterOutputConnection(&UVChannelIndex, &UVChannelIndex);
+}
+
+void FAlignUVMeshNode::Evaluate(UE::Dataflow::FContext& Context, const FDataflowOutput* Out) const
+{
+	using namespace UE::Geometry;
+	using namespace UE::MeshResizing;
+
+	if (Out->IsA(&UVChannelIndex))
+	{
+		SafeForwardInput(Context, &UVChannelIndex, &UVChannelIndex);
+		return;
+	}
+	else if (Out->IsA(&ResizingMesh))
+	{
+		if (TObjectPtr<UDataflowMesh> InResizingMesh = GetValue(Context, &ResizingMesh))
+		{
+			if (TObjectPtr<UDataflowMesh> InBaseMesh = GetValue(Context, &BaseMesh))
+			{
+				if (InResizingMesh->GetDynamicMesh() && InBaseMesh->GetDynamicMesh())
+				{
+					if (InResizingMesh->GetDynamicMeshRef().HasAttributes() && InBaseMesh->GetDynamicMeshRef().HasAttributes())
+					{
+						const int32 UVChannel = GetValue(Context, &UVChannelIndex);
+
+						if (UVChannel >= 0 && InResizingMesh->GetDynamicMeshRef().Attributes()->NumUVLayers() > UVChannel && InBaseMesh->GetDynamicMeshRef().Attributes()->NumUVLayers() > UVChannel)
+						{
+							const FDynamicMeshUVOverlay* const InResizedUVOverlay = InResizingMesh->GetDynamicMeshRef().Attributes()->GetUVLayer(UVChannel);
+
+							FDynamicMesh3 ResizedMesh;
+							ResizedMesh.Copy(InResizingMesh->GetDynamicMeshRef());
+							FDynamicMeshUVOverlay* const ResizedUVOverlay = ResizedMesh.Attributes()->GetUVLayer(UVChannel);
+
+							FMeshConnectedComponents UVIslands(&ResizedMesh);
+							UVIslands.FindConnectedTriangles([ResizedUVOverlay](int32 Triangle0, int32 Triangle1) 
+							{
+								return ResizedUVOverlay ? ResizedUVOverlay->AreTrianglesConnected(Triangle0, Triangle1) : false;
+							});
+
+							const FDynamicMeshUVOverlay* const BaseUVOverlay = InBaseMesh->GetDynamicMeshRef().Attributes()->GetUVLayer(UVChannel);
+
+							for (const FMeshConnectedComponents::FComponent& Island : UVIslands.Components)
+							{
+								if (!ensure(Island.Indices.Num() > 0))
+								{
+									continue;
+								}
+
+								TSet<int32> IslandVertices;
+								for (int32 TriID : Island.Indices)
+								{
+									FIndex3i Tri = ResizedMesh.GetTriangle(TriID);
+									IslandVertices.Add(Tri[0]);
+									IslandVertices.Add(Tri[1]);
+									IslandVertices.Add(Tri[2]);
+								}
+
+								TArray<FVector2f> ResizedUVs, BaseUVs;
+								TArray<int32> ElementIndices;
+								for (int32 ElementIndex : InResizedUVOverlay->ElementIndicesItr())
+								{
+									const int32 VertexForElement = InResizedUVOverlay->GetParentVertex(ElementIndex);
+									if (IslandVertices.Contains(VertexForElement))
+									{
+										ElementIndices.Add(ElementIndex);
+										ResizedUVs.Add(InResizedUVOverlay->GetElement(ElementIndex));
+										BaseUVs.Add(BaseUVOverlay->GetElement(ElementIndex));
+									}
+								}
+
+								const FTransform2f BestFit = Private::BestFitTransform(BaseUVs, ResizedUVs, bScale);
+								const FTransform2f InverseBestFit = BestFit.Inverse();
+
+								for (int32 UVIndex = 0; UVIndex < ResizedUVs.Num(); ++UVIndex)
+								{
+									const int32 ElementIndex = ElementIndices[UVIndex];
+									const int32 ElementVertex = InResizedUVOverlay->GetParentVertex(ElementIndex);
+									if (IslandVertices.Contains(ElementVertex))
+									{
+										const FVector2f UV = InverseBestFit.TransformPoint(ResizedUVs[UVIndex]);
+										ResizedUVOverlay->SetElement(ElementIndex, UV);
+									}
+								}
+							}
+
+							TObjectPtr<UDataflowMesh> OutResizedMesh = NewObject<UDataflowMesh>();
+							OutResizedMesh->SetDynamicMesh(MoveTemp(ResizedMesh));
+							OutResizedMesh->SetMaterials(InResizingMesh->GetMaterials());
+							SetValue(Context, OutResizedMesh, &ResizingMesh);
+							SetValue(Context, UVChannel, &UVChannelIndex);
+							return;
+						}
+						else
+						{
+							Context.Warning(TEXT("Invalid UVChannelIndex"), this, Out);
+						}
+					}
+					else
+					{
+						Context.Warning(TEXT("An input mesh does not have an AttributeSet"), this, Out);
+					}
+				}
+				else
+				{
+					Context.Warning(TEXT("An input mesh does not have a DynamicMesh object"), this, Out);
+				}
+			}
+		}
+		SafeForwardInput(Context, &ResizingMesh, &ResizingMesh);
+	}
+}
+
+
diff --git a/Engine/Plugins/Experimental/MeshResizing/Source/MeshResizingNodes/Private/MeshResizing/AlignUVMeshNode.h b/Engine/Plugins/Experimental/MeshResizing/Source/MeshResizingNodes/Private/MeshResizing/AlignUVMeshNode.h
new file mode 100644
--- /dev/null
+++ b/Engine/Plugins/Experimental/MeshResizing/Source/MeshResizingNodes/Private/MeshResizing/AlignUVMeshNode.h
@@ -0,0 +1,44 @@
+// Copyright Epic Games, Inc. All Rights Reserved.
+
+#pragma once
+
+#include "Dataflow/DataflowNode.h"
+#include "AlignUVMeshNode.generated.h"
+
+class UDataflowMesh;
+
+USTRUCT(Meta = (MeshResizing, Experimental))
+struct FAlignUVMeshNode : public FDataflowNode
+{
+	GENERATED_USTRUCT_BODY()
+	DATAFLOW_NODE_DEFINE_INTERNAL(FAlignUVMeshNode, "AlignUVMeshNode", "MeshResizing", "Align UV Mesh")
+	DATAFLOW_NODE_RENDER_TYPE("SurfaceRender", FName("UDataflowMesh"), "ResizingMesh", "UVChannelIndex")
+
+public:
+
+	FAlignUVMeshNode(const UE::Dataflow::FNodeParameters& InParam, FGuid InGuid = FGuid::NewGuid());
+
+private:
+
+	UPROPERTY(meta = (DataflowInput, DataflowOutput, DataflowPassthrough = "ResizingMesh"))
+	TObjectPtr<UDataflowMesh> ResizingMesh;
+
+	UPROPERTY(meta = (DataflowInput))
+	TObjectPtr<UDataflowMesh> BaseMesh;
+
+	UPROPERTY(EditAnywhere, Category = "Align", meta = (DataflowInput, DataflowOutput, DataflowPassthrough = "UVChannelIndex", ClampMin = 0, ClampMax = 7, UIMin = 0, UIMax = 7))
+	int32 UVChannelIndex = 0;
+
+	UPROPERTY(EditAnywhere, Category = "Align")
+	bool bScale = true;
+
+	//~ Begin FDataflowNode interface
+	virtual void Evaluate(UE::Dataflow::FContext& Context, const FDataflowOutput* Out) const override;
+	//~ End FDataflowNode interface
+};
+
+
+namespace UE::MeshResizing
+{
+	void RegisterAlignUVMeshNodes();
+}
diff --git a/Engine/Plugins/Experimental/MeshResizing/Source/MeshResizingNodes/Private/MeshResizing/MeshResizingDataflowNodesModule.cpp b/Engine/Plugins/Experimental/MeshResizing/Source/MeshResizingNodes/Private/MeshResizing/MeshResizingDataflowNodesModule.cpp
--- a/Engine/Plugins/Experimental/MeshResizing/Source/MeshResizingNodes/Private/MeshResizing/MeshResizingDataflowNodesModule.cpp
+++ b/Engine/Plugins/Experimental/MeshResizing/Source/MeshResizingNodes/Private/MeshResizing/MeshResizingDataflowNodesModule.cpp
@@ -7,6 +7,8 @@
 #include "MeshResizing/MeshConstraintNodes.h"
 #include "MeshResizing/MeshWarpNode.h"
 #include "MeshResizing/UVUnwrapNode.h"
+#include "MeshResizing/UVMeshTransformNode.h"
+#include "MeshResizing/AlignUVMeshNode.h"
 
 namespace UE::MeshResizing
 {
@@ -26,6 +28,8 @@ namespace UE::MeshResizing
 			RegisterBaseBodyDataflowNodes();
 			RegisterMeshConstraintDataflowNodes();
 			RegisterUVUnwrapNodes();
+			RegisterUVMeshTransformNodes();
+			RegisterAlignUVMeshNodes();
 
 			DATAFLOW_NODE_REGISTER_CREATION_FACTORY(FMeshWarpNode);
 		}
diff --git a/Engine/Plugins/Experimental/MeshResizing/Source/MeshResizingNodes/Private/MeshResizing/UVMeshTransformNode.cpp b/Engine/Plugins/Experimental/MeshResizing/Source/MeshResizingNodes/Private/MeshResizing/UVMeshTransformNode.cpp
new file mode 100644
--- /dev/null
+++ b/Engine/Plugins/Experimental/MeshResizing/Source/MeshResizingNodes/Private/MeshResizing/UVMeshTransformNode.cpp
@@ -0,0 +1,86 @@
+// Copyright Epic Games, Inc. All Rights Reserved.
+#include "MeshResizing/UVMeshTransformNode.h"
+#include "Dataflow/DataflowNodeFactory.h"
+#include "Dataflow/DataflowMesh.h"
+#include "UDynamicMesh.h"
+#include "Math/TransformCalculus2D.h"
+
+#include UE_INLINE_GENERATED_CPP_BY_NAME(UVMeshTransformNode)
+
+using UE::Geometry::FDynamicMesh3;
+
+namespace UE::MeshResizing
+{
+	void RegisterUVMeshTransformNodes()
+	{
+		DATAFLOW_NODE_REGISTER_CREATION_FACTORY(FUVMeshTransformNode);
+	}
+}
+
+FUVMeshTransformNode::FUVMeshTransformNode(const UE::Dataflow::FNodeParameters& InParam, FGuid InGuid)
+	: FDataflowNode(InParam, InGuid)
+{
+	RegisterInputConnection(&Mesh);
+	RegisterInputConnection(&UVChannelIndex);
+	RegisterOutputConnection(&Mesh, &Mesh);
+	RegisterOutputConnection(&UVChannelIndex, &UVChannelIndex);
+}
+
+void FUVMeshTransformNode::Evaluate(UE::Dataflow::FContext& Context, const FDataflowOutput* Out) const
+{
+	using namespace UE::Geometry;
+	using namespace UE::MeshResizing;
+
+	if (Out->IsA(&UVChannelIndex))
+	{
+		SafeForwardInput(Context, &UVChannelIndex, &UVChannelIndex);
+		return;
+	}
+	else if (Out->IsA(&Mesh))
+	{
+		if (const TObjectPtr<const UDataflowMesh> InMesh = GetValue(Context, &Mesh))
+		{
+			if (InMesh->GetDynamicMesh())
+			{
+				const int32 UVChannel = GetValue(Context, &UVChannelIndex);
+
+				if (InMesh->GetDynamicMeshRef().HasAttributes() && InMesh->GetDynamicMeshRef().Attributes()->NumUVLayers() > UVChannel && UVChannel >= 0)
+				{
+					const FDynamicMeshUVOverlay* const InMeshUVOverlay = InMesh->GetDynamicMeshRef().Attributes()->GetUVLayer(UVChannel);
+
+					FDynamicMesh3 OutMesh;
+					OutMesh.Copy(InMesh->GetDynamicMeshRef());
+					FDynamicMeshUVOverlay* const OutMeshUVOverlay = OutMesh.Attributes()->GetUVLayer(UVChannel);
+
+					checkf(InMeshUVOverlay->ElementCount() == OutMeshUVOverlay->ElementCount(), TEXT("Newly created UV overlay doesn't have the same size as input UV overlay"));
+
+					const FTransform2f Transform = FTransform2f{ FScale2f(Scale) }.Concatenate(FTransform2f{ FQuat2f{FMath::DegreesToRadians(Rotation)}, Translation });
+
+					for (int32 ElementIndex : InMeshUVOverlay->ElementIndicesItr())
+					{
+						const FVector2f NewUV = Transform.TransformPoint(InMeshUVOverlay->GetElement(ElementIndex));
+						OutMeshUVOverlay->SetElement(ElementIndex, NewUV);
+					}
+
+					TObjectPtr<UDataflowMesh> OutDataflowMesh = NewObject<UDataflowMesh>();
+					OutDataflowMesh->SetDynamicMesh(MoveTemp(OutMesh));
+					OutDataflowMesh->SetMaterials(InMesh->GetMaterials());
+					SetValue(Context, OutDataflowMesh, &Mesh);
+					SetValue(Context, UVChannel, &UVChannelIndex);
+					return;
+				}
+				else
+				{
+					Context.Warning(TEXT("Invalid UVChannelIndex or the input mesh does not have an AttributeSet"), this, Out);
+				}
+			}
+			else
+			{
+				Context.Warning(TEXT("Input mesh does not have a DynamicMesh object"), this, Out);
+			}
+		}
+
+		SafeForwardInput(Context, &Mesh, &Mesh);
+	}	
+}
+
diff --git a/Engine/Plugins/Experimental/MeshResizing/Source/MeshResizingNodes/Private/MeshResizing/UVMeshTransformNode.h b/Engine/Plugins/Experimental/MeshResizing/Source/MeshResizingNodes/Private/MeshResizing/UVMeshTransformNode.h
new file mode 100644
--- /dev/null
+++ b/Engine/Plugins/Experimental/MeshResizing/Source/MeshResizingNodes/Private/MeshResizing/UVMeshTransformNode.h
@@ -0,0 +1,49 @@
+// Copyright Epic Games, Inc. All Rights Reserved.
+
+#pragma once
+
+#include "Dataflow/DataflowNode.h"
+#include "UVMeshTransformNode.generated.h"
+
+class UDataflowMesh;
+
+USTRUCT(Meta = (MeshResizing, Experimental))
+struct FUVMeshTransformNode : public FDataflowNode
+{
+	GENERATED_USTRUCT_BODY()
+	DATAFLOW_NODE_DEFINE_INTERNAL(FUVMeshTransformNode, "UVMeshTransformNode", "MeshResizing", "UV Mesh Transform")
+	DATAFLOW_NODE_RENDER_TYPE("SurfaceRender", FName("UDataflowMesh"), "Mesh", "UVChannelIndex")
+
+public:
+
+	FUVMeshTransformNode(const UE::Dataflow::FNodeParameters& InParam, FGuid InGuid = FGuid::NewGuid());
+
+private:
+
+	UPROPERTY(meta = (DataflowInput, DataflowOutput, DataflowPassthrough = "Mesh"))
+	TObjectPtr<UDataflowMesh> Mesh;
+
+	UPROPERTY(EditAnywhere, Category = "Mesh", meta = (DataflowInput, DataflowOutput, DataflowPassthrough = "UVChannelIndex", ClampMin = 0, ClampMax = 7, UIMin = 0, UIMax = 7))
+	int32 UVChannelIndex = 0;
+
+	UPROPERTY(EditAnywhere, Category = "Transform", Meta = (AllowPreserveRatio))
+	FVector2f Scale = { 1.f, 1.f };
+
+	/** Rotation angle in degrees */
+	UPROPERTY(EditAnywhere, Category = "Transform", Meta = (UIMin = -360, UIMax = 360, ClampMin = -360, ClampMax = 360))
+	float Rotation = 0.f;
+
+	UPROPERTY(EditAnywhere, Category = "Transform")
+	FVector2f Translation = { 0.f, 0.f };
+
+	//~ Begin FDataflowNode interface
+	virtual void Evaluate(UE::Dataflow::FContext& Context, const FDataflowOutput* Out) const override;
+	//~ End FDataflowNode interface
+};
+
+
+
+namespace UE::MeshResizing
+{
+	void RegisterUVMeshTransformNodes();
+}
