Mesh Resizing: Add new RBFInterpolation nodes which allow you to sample and calculate weights with one node, and then apply the interpolation with another node. Added the option to interpolate normals.

diff --git a/Engine/Plugins/Experimental/MeshResizing/Source/MeshResizingCore/MeshResizingCore.Build.cs b/Engine/Plugins/Experimental/MeshResizing/Source/MeshResizingCore/MeshResizingCore.Build.cs
--- a/Engine/Plugins/Experimental/MeshResizing/Source/MeshResizingCore/MeshResizingCore.Build.cs
+++ b/Engine/Plugins/Experimental/MeshResizing/Source/MeshResizingCore/MeshResizingCore.Build.cs
@@ -12,6 +12,7 @@ namespace UnrealBuildTool.Rules
 				new string[]
 				{
 					"Core",
+					"CoreUObject",
 					"GeometryCore",
 					"HairStrandsCore",
 				}
diff --git a/Engine/Plugins/Experimental/MeshResizing/Source/MeshResizingCore/Private/MeshResizing/RBFInterpolation.cpp b/Engine/Plugins/Experimental/MeshResizing/Source/MeshResizingCore/Private/MeshResizing/RBFInterpolation.cpp
--- a/Engine/Plugins/Experimental/MeshResizing/Source/MeshResizingCore/Private/MeshResizing/RBFInterpolation.cpp
+++ b/Engine/Plugins/Experimental/MeshResizing/Source/MeshResizingCore/Private/MeshResizing/RBFInterpolation.cpp
@@ -2,14 +2,17 @@
 
 #include "MeshResizing/RBFInterpolation.h"
 #include "DynamicMesh/DynamicMesh3.h"
+#include "DynamicMesh/DynamicMeshAttributeSet.h"
 #include "GroomBindingBuilder.h"
 #include "GroomRBFDeformer.h"
 
+#include UE_INLINE_GENERATED_CPP_BY_NAME(RBFInterpolation)
+
 namespace UE::MeshResizing
 {
 	namespace Private
 	{
-		static void UpdateInterpolationData(const GroomBinding_RBFWeighting::FPointsSampler& PointsSampler, const GroomBinding_RBFWeighting::FWeightsBuilder& InterpolationWeights, FRBFInterpolation::FInterpolationData& OutData)
+		static void UpdateInterpolationData(const GroomBinding_RBFWeighting::FPointsSampler& PointsSampler, const GroomBinding_RBFWeighting::FWeightsBuilder& InterpolationWeights, FMeshResizingRBFInterpolationData& OutData)
 		{
 			OutData.SampleIndices = PointsSampler.SampleIndices;
 			OutData.SampleRestPositions = PointsSampler.SamplePositions;
@@ -17,7 +20,7 @@ namespace UE::MeshResizing
 		}
 
 		// From GroomRBFDeformer
-		static void UpdateMeshSamples(const TArray<FVector3f>& TargetPositions, const FRBFInterpolation::FInterpolationData& InterpolationData, TArray<FVector3f>& OutSampleDeformations)
+		static void UpdateMeshSamples(const TArray<FVector3f>& TargetPositions, const FMeshResizingRBFInterpolationData& InterpolationData, TArray<FVector3f>& OutSampleDeformations)
 		{
 			const uint32 MaxSampleCount = TargetPositions.Num();
 			check(MaxSampleCount == InterpolationData.SampleIndices.Num());
@@ -39,31 +42,68 @@ namespace UE::MeshResizing
 		}
 
 		// From GroomRBFDeformer
-		static void ApplyRBFDeformation(const FRBFInterpolation::FInterpolationData& InterpolationData, const TArray<FVector3f>& SampleDeformations, UE::Geometry::FDynamicMesh3& DeformingMesh)
+		static void ApplyRBFDeformation(const FMeshResizingRBFInterpolationData& InterpolationData, const TArray<FVector3f>& SampleDeformations, bool bInterpolateNormals, UE::Geometry::FDynamicMesh3& DeformingMesh)
 		{
 			const uint32 SampleCount = InterpolationData.SampleRestPositions.Num();
 			check(SampleCount + 4 == SampleDeformations.Num());
-			for (const int32 VertexIndex : DeformingMesh.VertexIndicesItr())
+
+			TArray<UE::Geometry::FDynamicMeshNormalOverlay*> NormalOverlays;
+			if (bInterpolateNormals)
 			{
-				const FVector3f RestControlPoint(FVector3f(DeformingMesh.GetVertexRef(VertexIndex)));
-				FVector3f ControlPoint = RestControlPoint;
-				for (uint32 SampleIndex = 0; SampleIndex < SampleCount; ++SampleIndex)
+				const int32 NumNormalLayers = DeformingMesh.Attributes() ? DeformingMesh.Attributes()->NumNormalLayers() : 0;
+				NormalOverlays.Reserve(NumNormalLayers);
+
+				for (int32 NormalLayerId = 0; NormalLayerId < NumNormalLayers; ++NormalLayerId)
 				{
-					const FVector3f PositionDelta = RestControlPoint - InterpolationData.SampleRestPositions[SampleIndex];
-					const float FunctionValue = FMath::Sqrt(FVector3f::DotProduct(PositionDelta, PositionDelta) + 1);
-					ControlPoint += FunctionValue * SampleDeformations[SampleIndex];
+					if (UE::Geometry::FDynamicMeshNormalOverlay* const Overlay = DeformingMesh.Attributes()->GetNormalLayer(NormalLayerId))
+					{
+						NormalOverlays.Emplace(Overlay);
+					}
 				}
-				ControlPoint += SampleDeformations[SampleCount];
-				ControlPoint += SampleDeformations[SampleCount + 1] * RestControlPoint.X;
-				ControlPoint += SampleDeformations[SampleCount + 2] * RestControlPoint.Y;
-				ControlPoint += SampleDeformations[SampleCount + 3] * RestControlPoint.Z;
+			}
 
-				DeformingMesh.SetVertex(VertexIndex, FVector3d(ControlPoint));
+			for (const int32 VertexIndex : DeformingMesh.VertexIndicesItr())
+			{
+				auto DeformPoint = [SampleCount, &SampleDeformations, &InterpolationData](const FVector3f& RestControlPoint)
+					{
+						FVector3f ControlPoint = RestControlPoint;
+						for (uint32 SampleIndex = 0; SampleIndex < SampleCount; ++SampleIndex)
+						{
+							const FVector3f PositionDelta = RestControlPoint - InterpolationData.SampleRestPositions[SampleIndex];
+							const float FunctionValue = FMath::Sqrt(FVector3f::DotProduct(PositionDelta, PositionDelta) + 1);
+							ControlPoint += FunctionValue * SampleDeformations[SampleIndex];
+						}
+						ControlPoint += SampleDeformations[SampleCount];
+						ControlPoint += SampleDeformations[SampleCount + 1] * RestControlPoint.X;
+						ControlPoint += SampleDeformations[SampleCount + 2] * RestControlPoint.Y;
+						ControlPoint += SampleDeformations[SampleCount + 3] * RestControlPoint.Z;
+
+						return ControlPoint;
+					};
+
+				const FVector3f RestPoint(FVector3f(DeformingMesh.GetVertexRef(VertexIndex)));
+				const FVector3f DeformedPoint = DeformPoint(RestPoint);
+
+				DeformingMesh.SetVertex(VertexIndex, FVector3d(DeformedPoint));
+
+				for (UE::Geometry::FDynamicMeshNormalOverlay* const Overlay : NormalOverlays)
+				{
+					Overlay->EnumerateVertexElements(VertexIndex,
+						[&RestPoint, &DeformedPoint, &DeformPoint, Overlay](int32 TriangleID, int32 ElementID, const FVector3f& Normal)
+						{
+							const FVector3f RestNormalEnd = RestPoint + Normal;
+							const FVector3f DeformedNormalEnd = DeformPoint(RestNormalEnd);
+							const FVector3f DeformedNorml = (DeformedNormalEnd - DeformedPoint).GetSafeNormal();
+							Overlay->SetElement(ElementID, DeformedNorml);
+							return true;
+						}
+					);
+				}
 			}
 		}		
 	}
 
-	void FRBFInterpolation::GenerateWeights(const UE::Geometry::FDynamicMesh3& BaseMesh, int32 NumInterpolationPoints, FInterpolationData& OutData)
+	void FRBFInterpolation::GenerateWeights(const UE::Geometry::FDynamicMesh3& BaseMesh, int32 NumInterpolationPoints, FMeshResizingRBFInterpolationData& OutData)
 	{
 		// Groom binding sampler expects points to be a contiguous array of FVector3f. Just copy that way for now.
 		TArray<FVector3f> Positions;
@@ -83,10 +123,10 @@ namespace UE::MeshResizing
 		Private::UpdateInterpolationData(PointsSampler, InterpolationWeights, OutData);
 	}
 
-	void FRBFInterpolation::DeformPoints(const TArray<FVector3f>& TargetPositions, const FInterpolationData& InterpolationData, UE::Geometry::FDynamicMesh3& DeformingMesh)
+	void FRBFInterpolation::DeformPoints(const TArray<FVector3f>& TargetPositions, const FMeshResizingRBFInterpolationData& InterpolationData, bool bInterpolateNormals, UE::Geometry::FDynamicMesh3& DeformingMesh)
 	{
 		TArray<FVector3f> SampleDeformations;
 		Private::UpdateMeshSamples(TargetPositions, InterpolationData, SampleDeformations);
-		Private::ApplyRBFDeformation(InterpolationData, SampleDeformations, DeformingMesh);
+		Private::ApplyRBFDeformation(InterpolationData, SampleDeformations, bInterpolateNormals, DeformingMesh);
 	}
 }
\ No newline at end of file
diff --git a/Engine/Plugins/Experimental/MeshResizing/Source/MeshResizingCore/Public/MeshResizing/RBFInterpolation.h b/Engine/Plugins/Experimental/MeshResizing/Source/MeshResizingCore/Public/MeshResizing/RBFInterpolation.h
--- a/Engine/Plugins/Experimental/MeshResizing/Source/MeshResizingCore/Public/MeshResizing/RBFInterpolation.h
+++ b/Engine/Plugins/Experimental/MeshResizing/Source/MeshResizingCore/Public/MeshResizing/RBFInterpolation.h
@@ -3,25 +3,33 @@
 #pragma once 
 #include "UObject/NameTypes.h"
 #include "Containers/Array.h"
+#include "RBFInterpolation.generated.h"
 
 namespace UE::Geometry
 {
 	class FDynamicMesh3;
 }
+
+USTRUCT()
+struct MESHRESIZINGCORE_API FMeshResizingRBFInterpolationData
+{
+	GENERATED_BODY()
+
+	UPROPERTY()
+	TArray<uint32> SampleIndices;
+	UPROPERTY()
+	TArray<FVector3f> SampleRestPositions;
+	UPROPERTY()
+	TArray<float> InterpolationWeights;
+};
+
 namespace UE::MeshResizing
 {
 	struct MESHRESIZINGCORE_API FRBFInterpolation
 	{
 
-		struct FInterpolationData
-		{
-			TArray<uint32> SampleIndices;
-			TArray<FVector3f> SampleRestPositions;
-			TArray<float> InterpolationWeights;
-		};
-
-		static void GenerateWeights(const UE::Geometry::FDynamicMesh3& BaseMesh, int32 NumInterpolationPoints, FInterpolationData& OutData);
-		static void DeformPoints(const TArray<FVector3f>& TargetPositions, const FInterpolationData& InterpolationData, UE::Geometry::FDynamicMesh3& DeformingMesh);
+		static void GenerateWeights(const UE::Geometry::FDynamicMesh3& BaseMesh, int32 NumInterpolationPoints, FMeshResizingRBFInterpolationData& OutData);
+		static void DeformPoints(const TArray<FVector3f>& TargetPositions, const FMeshResizingRBFInterpolationData& InterpolationData, bool bInterpolateNormals,  UE::Geometry::FDynamicMesh3& DeformingMesh);
 	};
 }
 
diff --git a/Engine/Plugins/Experimental/MeshResizing/Source/MeshResizingNodes/Private/MeshResizing/MeshResizingDataflowNodesModule.cpp b/Engine/Plugins/Experimental/MeshResizing/Source/MeshResizingNodes/Private/MeshResizing/MeshResizingDataflowNodesModule.cpp
--- a/Engine/Plugins/Experimental/MeshResizing/Source/MeshResizingNodes/Private/MeshResizing/MeshResizingDataflowNodesModule.cpp
+++ b/Engine/Plugins/Experimental/MeshResizing/Source/MeshResizingNodes/Private/MeshResizing/MeshResizingDataflowNodesModule.cpp
@@ -6,6 +6,7 @@
 #include "MeshResizing/BaseBodyDataflowNodes.h"
 #include "MeshResizing/MeshConstraintNodes.h"
 #include "MeshResizing/MeshWarpNode.h"
+#include "MeshResizing/RBFInterpolationNodes.h"
 #include "MeshResizing/UVUnwrapNode.h"
 #include "MeshResizing/UVMeshTransformNode.h"
 #include "MeshResizing/AlignUVMeshNode.h"
@@ -32,6 +33,8 @@ namespace UE::MeshResizing
 			RegisterAlignUVMeshNodes();
 
 			DATAFLOW_NODE_REGISTER_CREATION_FACTORY(FMeshWarpNode);
+			DATAFLOW_NODE_REGISTER_CREATION_FACTORY(FGenerateRBFResizingWeightsNode);
+			DATAFLOW_NODE_REGISTER_CREATION_FACTORY(FApplyRBFResizingNode);
 		}
 
 	}  // End namespace Private
diff --git a/Engine/Plugins/Experimental/MeshResizing/Source/MeshResizingNodes/Private/MeshResizing/MeshWarpNode.cpp b/Engine/Plugins/Experimental/MeshResizing/Source/MeshResizingNodes/Private/MeshResizing/MeshWarpNode.cpp
--- a/Engine/Plugins/Experimental/MeshResizing/Source/MeshResizingNodes/Private/MeshResizing/MeshWarpNode.cpp
+++ b/Engine/Plugins/Experimental/MeshResizing/Source/MeshResizingNodes/Private/MeshResizing/MeshWarpNode.cpp
@@ -98,14 +98,15 @@ namespace  UE::MeshResizing::Private
 		FDynamicMesh3& ResizedMesh,
 		const FDynamicMesh3& SourceMesh,
 		const FDynamicMesh3& TargetMesh,
-		int32 NumInterpolationPoints)
+		int32 NumInterpolationPoints,
+		bool bInterpolateNormals)
 	{
 		if (ResizedMesh.VertexCount() == 0)
 		{
 			return;
 		}
 
-		FRBFInterpolation::FInterpolationData InterpolationData;
+		FMeshResizingRBFInterpolationData InterpolationData;
 		FRBFInterpolation::GenerateWeights(SourceMesh, NumInterpolationPoints, InterpolationData);
 
 		// Get target mesh samples
@@ -116,7 +117,7 @@ namespace  UE::MeshResizing::Private
 			TargetSamplePoints.Emplace(FVector3f(TargetMesh.GetVertex(SampleIndex)));
 		}
 		// Do the interpolation
-		FRBFInterpolation::DeformPoints(TargetSamplePoints, InterpolationData, ResizedMesh);
+		FRBFInterpolation::DeformPoints(TargetSamplePoints, InterpolationData, bInterpolateNormals, ResizedMesh);
 	}
 }
 
@@ -156,7 +157,7 @@ void FMeshWarpNode::Evaluate(UE::Dataflow::FContext& Context, const FDataflowOut
 					{
 					case EMeshResizingWarpMethod::RBFInterpolate:
 					default:
-						UE::MeshResizing::Private::ApplyRefitToMesh_RBFInterpolate(ResizeMesh, InSourceMesh->GetDynamicMeshRef(), BlendedTargetFMesh, NumInterpolationPoints);
+						UE::MeshResizing::Private::ApplyRefitToMesh_RBFInterpolate(ResizeMesh, InSourceMesh->GetDynamicMeshRef(), BlendedTargetFMesh, NumInterpolationPoints, bInterpolateNormals);
 						break;
 					case EMeshResizingWarpMethod::WrapDeform:
 						UE::MeshResizing::Private::ApplyRefitToMesh_WrapDeform(ResizeMesh, InSourceMesh->GetDynamicMeshRef(), BlendedTargetFMesh, 1.0f);
@@ -178,4 +179,3 @@ void FMeshWarpNode::Evaluate(UE::Dataflow::FContext& Context, const FDataflowOut
 	SafeForwardInput(Context, &MeshToResize, &ResizedMesh);
 	SafeForwardInput(Context, &SourceMesh, &BlendedTargetMesh);
 }
-
diff --git a/Engine/Plugins/Experimental/MeshResizing/Source/MeshResizingNodes/Private/MeshResizing/MeshWarpNode.h b/Engine/Plugins/Experimental/MeshResizing/Source/MeshResizingNodes/Private/MeshResizing/MeshWarpNode.h
--- a/Engine/Plugins/Experimental/MeshResizing/Source/MeshResizingNodes/Private/MeshResizing/MeshWarpNode.h
+++ b/Engine/Plugins/Experimental/MeshResizing/Source/MeshResizingNodes/Private/MeshResizing/MeshWarpNode.h
@@ -58,7 +58,8 @@ private:
 	UPROPERTY(EditAnywhere, Category = "Warp", meta =  (Min = 1, EditCondition = "WarpMethod == EMeshResizingWarpMethod::RBFInterpolate", EditConditionHides))
 	int32 NumInterpolationPoints = 100;
 
-
+	UPROPERTY(EditAnywhere, Category = "Warp", meta = (EditCondition = "WarpMethod == EMeshResizingWarpMethod::RBFInterpolate", EditConditionHides))
+	bool bInterpolateNormals = true;
 
 	//~ Begin FDataflowNode interface
 	virtual void Evaluate(UE::Dataflow::FContext& Context, const FDataflowOutput* Out) const override;
diff --git a/Engine/Plugins/Experimental/MeshResizing/Source/MeshResizingNodes/Private/MeshResizing/RBFInterpolationNodes.cpp b/Engine/Plugins/Experimental/MeshResizing/Source/MeshResizingNodes/Private/MeshResizing/RBFInterpolationNodes.cpp
new file mode 100644
--- /dev/null
+++ b/Engine/Plugins/Experimental/MeshResizing/Source/MeshResizingNodes/Private/MeshResizing/RBFInterpolationNodes.cpp
@@ -0,0 +1,120 @@
+// Copyright Epic Games, Inc. All Rights Reserved.
+#include "MeshResizing/RBFInterpolationNodes.h"
+#include "MeshResizing/RBFInterpolation.h"
+#include "DynamicMesh/DynamicMesh3.h"
+#include "DynamicMesh/DynamicMeshAABBTree3.h"
+#include "DynamicMeshEditor.h"
+#include "UDynamicMesh.h"
+#include "Dataflow/DataflowDebugDrawInterface.h"
+#include "Dataflow/DataflowMesh.h"
+#include "Materials/MaterialInterface.h"
+#if WITH_EDITOR
+#include "Dataflow/DataflowRenderingViewMode.h"
+#endif
+
+#include UE_INLINE_GENERATED_CPP_BY_NAME(RBFInterpolationNodes)
+
+using UE::Geometry::FDynamicMesh3;
+
+FGenerateRBFResizingWeightsNode::FGenerateRBFResizingWeightsNode(const UE::Dataflow::FNodeParameters& InParam, FGuid InGuid) :
+	FDataflowNode(InParam, InGuid)
+{
+	RegisterInputConnection(&MeshToResize);
+	RegisterInputConnection(&SourceMesh);
+
+	RegisterOutputConnection(&InterpolationData);
+}
+
+void FGenerateRBFResizingWeightsNode::Evaluate(UE::Dataflow::FContext& Context, const FDataflowOutput* Out) const
+{
+	using namespace UE::Geometry;
+	if (Out->IsA<FMeshResizingRBFInterpolationData>(&InterpolationData))
+	{
+		FMeshResizingRBFInterpolationData Result;
+		if (TObjectPtr<UDataflowMesh> InMeshToResize = GetValue(Context, &MeshToResize))
+		{
+			if (TObjectPtr<UDataflowMesh> InSourceMesh = GetValue(Context, &SourceMesh))
+			{
+				if (InMeshToResize->GetDynamicMesh() && InSourceMesh->GetDynamicMesh())
+				{
+					UE::MeshResizing::FRBFInterpolation::GenerateWeights(InSourceMesh->GetDynamicMeshRef(), NumInterpolationPoints, Result);
+				}
+			}
+		}
+
+		SetValue(Context, MoveTemp(Result), &InterpolationData);
+	}
+}
+
+#if WITH_EDITOR
+bool FGenerateRBFResizingWeightsNode::CanDebugDrawViewMode(const FName& ViewMode) const
+{
+	return ViewMode == UE::Dataflow::FDataflowConstruction3DViewMode::Name;
+}
+
+void FGenerateRBFResizingWeightsNode::DebugDraw(UE::Dataflow::FContext& Context, IDataflowDebugDrawInterface& DataflowRenderingInterface, const FDebugDrawParameters& DebugDrawParameters) const
+{
+	if (DebugDrawParameters.bNodeIsSelected || DebugDrawParameters.bNodeIsPinned)
+	{
+		const FMeshResizingRBFInterpolationData& OutInterpolationData = GetOutputValue(Context, &InterpolationData, InterpolationData);
+		if (OutInterpolationData.SampleIndices.Num())
+		{
+			DataflowRenderingInterface.SetColor(FLinearColor::Yellow);
+			DataflowRenderingInterface.SetPointSize(2.f);
+
+			DataflowRenderingInterface.ReservePoints(OutInterpolationData.SampleRestPositions.Num());
+			for (const FVector3f& Pos : OutInterpolationData.SampleRestPositions)
+			{
+				DataflowRenderingInterface.DrawPoint(FVector(Pos));
+			}
+		}
+	}
+}
+#endif
+
+FApplyRBFResizingNode::FApplyRBFResizingNode(const UE::Dataflow::FNodeParameters& InParam, FGuid InGuid) :
+	FDataflowNode(InParam, InGuid)
+{
+	RegisterInputConnection(&MeshToResize);
+	RegisterInputConnection(&TargetMesh);
+	RegisterInputConnection(&InterpolationData);
+
+	RegisterOutputConnection(&ResizedMesh, &MeshToResize);
+}
+
+
+void FApplyRBFResizingNode::Evaluate(UE::Dataflow::FContext& Context, const FDataflowOutput* Out) const
+{
+	if (TObjectPtr<UDataflowMesh> InMeshToResize = GetValue(Context, &MeshToResize))
+	{
+		if (TObjectPtr<UDataflowMesh> InTargetMesh = GetValue(Context, &TargetMesh))
+		{
+			const FMeshResizingRBFInterpolationData& InInterpolationData = GetValue(Context, &InterpolationData);
+
+			if (InMeshToResize->GetDynamicMesh() && InTargetMesh->GetDynamicMesh() && InInterpolationData.SampleIndices.Num())
+			{
+				TObjectPtr<UDataflowMesh> OutResizedMesh = NewObject<UDataflowMesh>();
+				FDynamicMesh3 ResizeMesh;
+				ResizeMesh.Copy(InMeshToResize->GetDynamicMeshRef());
+
+				const FDynamicMesh3& TargetMeshRef = InTargetMesh->GetDynamicMeshRef();
+
+				// Get target mesh samples
+				TArray<FVector3f> TargetSamplePoints;
+				TargetSamplePoints.Reserve(InInterpolationData.SampleIndices.Num());
+				for (const uint32 SampleIndex : InInterpolationData.SampleIndices)
+				{
+					TargetSamplePoints.Emplace(FVector3f(TargetMeshRef.GetVertex(SampleIndex)));
+				}
+				// Do the interpolation
+				UE::MeshResizing::FRBFInterpolation::DeformPoints(TargetSamplePoints, InInterpolationData, bInterpolateNormals, ResizeMesh);
+				OutResizedMesh->SetDynamicMesh(MoveTemp(ResizeMesh));
+				OutResizedMesh->SetMaterials(InMeshToResize->GetMaterials());
+				SetValue(Context, OutResizedMesh, &ResizedMesh);
+				return;
+			}
+		}
+	}
+
+	SafeForwardInput(Context, &MeshToResize, &ResizedMesh);
+}
\ No newline at end of file
diff --git a/Engine/Plugins/Experimental/MeshResizing/Source/MeshResizingNodes/Private/MeshResizing/RBFInterpolationNodes.h b/Engine/Plugins/Experimental/MeshResizing/Source/MeshResizingNodes/Private/MeshResizing/RBFInterpolationNodes.h
new file mode 100644
--- /dev/null
+++ b/Engine/Plugins/Experimental/MeshResizing/Source/MeshResizingNodes/Private/MeshResizing/RBFInterpolationNodes.h
@@ -0,0 +1,96 @@
+// Copyright Epic Games, Inc. All Rights Reserved.
+
+#pragma once
+
+#include "Dataflow/DataflowNode.h"
+#include "GeometryCollection/ManagedArrayCollection.h"
+#include "MeshResizing/RBFInterpolation.h"
+
+#include "RBFInterpolationNodes.generated.h"
+
+class UDataflowMesh;
+class UMaterialInterface;
+
+/**
+ * Sample points and generate RBF Interpolation data for a given Source mesh.
+ */
+USTRUCT(Meta = (MeshResizing, Experimental))
+struct FGenerateRBFResizingWeightsNode : public FDataflowNode
+{
+	GENERATED_USTRUCT_BODY()
+	DATAFLOW_NODE_DEFINE_INTERNAL(FGenerateRBFResizingWeightsNode, "GenerateRBFResizingWeights", "MeshResizing", "Generate RBF Resizing Weights")
+
+public:
+
+	FGenerateRBFResizingWeightsNode(const UE::Dataflow::FNodeParameters& InParam, FGuid InGuid = FGuid::NewGuid());
+
+private:
+
+	/** The mesh to resize. This is currently unused but may be used to improve point sampling in the future. */
+	UPROPERTY(meta = (DataflowInput))
+	TObjectPtr<UDataflowMesh> MeshToResize;
+
+	/** The source mesh to be sampled.*/
+	UPROPERTY(meta = (DataflowInput))
+	TObjectPtr<UDataflowMesh> SourceMesh;
+
+	/** The calculated interpolation points and RBF weights */
+	UPROPERTY(meta = (DataflowOutput))
+	FMeshResizingRBFInterpolationData InterpolationData;
+
+	/** The number of interpolation points to be sampled. */
+	UPROPERTY(EditAnywhere, Category = "RBF Resizing Weights", meta = (Min = 1))
+	int32 NumInterpolationPoints = 1000;
+
+	//~ Begin FDataflowNode interface
+	virtual void Evaluate(UE::Dataflow::FContext& Context, const FDataflowOutput* Out) const override; 
+#if WITH_EDITOR
+	virtual bool CanDebugDraw() const override
+	{
+		return true;
+	}
+	virtual bool CanDebugDrawViewMode(const FName& ViewModeName) const override;
+	virtual void DebugDraw(UE::Dataflow::FContext& Context, IDataflowDebugDrawInterface& DataflowRenderingInterface, const FDebugDrawParameters& DebugDrawParameters) const override;
+#endif
+	//~ End FDataflowNode interface
+};
+
+/**
+ * Apply the interpolation data calculated by GenerateRBFResizingWeights to resize a mesh.
+ */
+USTRUCT(Meta = (MeshResizing, Experimental))
+struct FApplyRBFResizingNode : public FDataflowNode
+{
+	GENERATED_USTRUCT_BODY()
+	DATAFLOW_NODE_DEFINE_INTERNAL(FApplyRBFResizingNode, "ApplyRBFResizing", "MeshResizing", "Apply RBF Resizing")
+	DATAFLOW_NODE_RENDER_TYPE("SurfaceRender", FName("UDataflowMesh"), "ResizedMesh")
+
+public:
+	FApplyRBFResizingNode(const UE::Dataflow::FNodeParameters& InParam, FGuid InGuid = FGuid::NewGuid());
+
+private:
+
+	/** The mesh being resized */
+	UPROPERTY(meta = (DataflowInput))
+	TObjectPtr<UDataflowMesh> MeshToResize;
+
+	/** The target mesh that corresponds with the SourceMesh used to generate the InterpolationData. Must have matching vertices with SourceMesh */
+	UPROPERTY(meta = (DataflowInput))
+	TObjectPtr<UDataflowMesh> TargetMesh;
+
+	/** The pre-calculated base RBF interpolation data.*/
+	UPROPERTY(meta = (DataflowInput))
+	FMeshResizingRBFInterpolationData InterpolationData;
+
+	/** The resulting resized mesh */
+	UPROPERTY(meta = (DataflowOutput, DataflowPassthrough = "MeshToResize"))
+	TObjectPtr<UDataflowMesh> ResizedMesh;
+
+	/** Whether or not to interpolate the normals as well as the positions*/
+	UPROPERTY(EditAnywhere, Category = "RBF Resize")
+	bool bInterpolateNormals = true;
+
+	//~ Begin FDataflowNode interface
+	virtual void Evaluate(UE::Dataflow::FContext& Context, const FDataflowOutput* Out) const override;
+	//~ End FDataflowNode interface
+};
\ No newline at end of file
