Optimize Mesh resizing using tasks
28 times faster on the apply deformation function

diff --git a/Engine/Plugins/Experimental/MeshResizing/Source/MeshResizingCore/Private/MeshResizing/RBFInterpolation.cpp b/Engine/Plugins/Experimental/MeshResizing/Source/MeshResizingCore/Private/MeshResizing/RBFInterpolation.cpp
--- a/Engine/Plugins/Experimental/MeshResizing/Source/MeshResizingCore/Private/MeshResizing/RBFInterpolation.cpp
+++ b/Engine/Plugins/Experimental/MeshResizing/Source/MeshResizingCore/Private/MeshResizing/RBFInterpolation.cpp
@@ -5,6 +5,7 @@
 #include "DynamicMesh/DynamicMeshAttributeSet.h"
 #include "GroomBindingBuilder.h"
 #include "GroomRBFDeformer.h"
+#include "Tasks/Task.h"
 
 #include UE_INLINE_GENERATED_CPP_BY_NAME(RBFInterpolation)
 
@@ -29,16 +30,35 @@ namespace UE::MeshResizing
 			const uint32 EntryCount = FGroomRBFDeformer::GetEntryCount(MaxSampleCount);
 			OutSampleDeformations.SetNum(EntryCount);
 
-			for (uint32 SampleIndex = 0; SampleIndex < EntryCount; ++SampleIndex)
+			const uint32 NumTasks = FMath::Max(FMath::Min(uint32(FTaskGraphInterface::Get().GetNumWorkerThreads()), EntryCount), 1u);
+			constexpr uint32 MinEntryByTask = 10;
+			const uint32 EntryByTask = FMath::Max(FMath::Min(FMath::DivideAndRoundUp(EntryCount, NumTasks), MinEntryByTask), 1u);
+			const uint32 NumBatches = FMath::DivideAndRoundUp(EntryCount, EntryByTask);
+			TArray<UE::Tasks::FTask> PendingTasks;
+			PendingTasks.Reserve(NumBatches);
+
+			for (uint32 BatchIndex = 0; BatchIndex < NumBatches; BatchIndex++)
 			{
-				uint32 WeightsOffset = SampleIndex * EntryCount;
-				FVector3f SampleDeformation(FVector3f::ZeroVector);
-				for (uint32 Index = 0; Index < MaxSampleCount; ++Index, ++WeightsOffset)
+				const uint32 StartIndex = BatchIndex * EntryByTask;
+				uint32 EndIndex = (BatchIndex + 1) * EntryByTask;
+				EndIndex = BatchIndex == NumBatches - 1 ? FMath::Min(EntryCount, EndIndex) : EndIndex;
+
+				UE::Tasks::FTask PendingTask = UE::Tasks::Launch(UE_SOURCE_LOCATION, [StartIndex, EndIndex, EntryCount, MaxSampleCount, &OutSampleDeformations, &InterpolationData, &TargetPositions]
 				{
-					SampleDeformation += InterpolationData.InterpolationWeights[WeightsOffset] * (TargetPositions[Index] - InterpolationData.SampleRestPositions[Index]);
-				}
-				OutSampleDeformations[SampleIndex] = SampleDeformation;
+					for (uint32 SampleIndex = StartIndex; SampleIndex < EndIndex; ++SampleIndex)
+					{
+						uint32 WeightsOffset = SampleIndex * EntryCount;
+						FVector3f SampleDeformation(FVector3f::ZeroVector);
+						for (uint32 Index = 0; Index < MaxSampleCount; ++Index, ++WeightsOffset)
+						{
+							SampleDeformation += InterpolationData.InterpolationWeights[WeightsOffset] * (TargetPositions[Index] - InterpolationData.SampleRestPositions[Index]);
+						}
+						OutSampleDeformations[SampleIndex] = SampleDeformation;
+					}
+				});
+					PendingTasks.Add(PendingTask);
 			}
+			UE::Tasks::Wait(PendingTasks);
 		}
 
 		// From GroomRBFDeformer
@@ -62,45 +82,66 @@ namespace UE::MeshResizing
 				}
 			}
 
-			for (const int32 VertexIndex : DeformingMesh.VertexIndicesItr())
+			auto DeformPoint = [SampleCount, &SampleDeformations, &InterpolationData](const FVector3f& RestControlPoint)
 			{
-				auto DeformPoint = [SampleCount, &SampleDeformations, &InterpolationData](const FVector3f& RestControlPoint)
-					{
-						FVector3f ControlPoint = RestControlPoint;
-						for (uint32 SampleIndex = 0; SampleIndex < SampleCount; ++SampleIndex)
-						{
-							const FVector3f PositionDelta = RestControlPoint - InterpolationData.SampleRestPositions[SampleIndex];
-							const float FunctionValue = FMath::Sqrt(FVector3f::DotProduct(PositionDelta, PositionDelta) + 1);
-							ControlPoint += FunctionValue * SampleDeformations[SampleIndex];
-						}
-						ControlPoint += SampleDeformations[SampleCount];
-						ControlPoint += SampleDeformations[SampleCount + 1] * RestControlPoint.X;
-						ControlPoint += SampleDeformations[SampleCount + 2] * RestControlPoint.Y;
-						ControlPoint += SampleDeformations[SampleCount + 3] * RestControlPoint.Z;
-
-						return ControlPoint;
-					};
+				FVector3f ControlPoint = RestControlPoint;
+				for (uint32 SampleIndex = 0; SampleIndex < SampleCount; ++SampleIndex)
+				{
+					const FVector3f PositionDelta = RestControlPoint - InterpolationData.SampleRestPositions[SampleIndex];
+					const float FunctionValue = FMath::Sqrt(FVector3f::DotProduct(PositionDelta, PositionDelta) + 1);
+					ControlPoint += FunctionValue * SampleDeformations[SampleIndex];
+				}
+				ControlPoint += SampleDeformations[SampleCount];
+				ControlPoint += SampleDeformations[SampleCount + 1] * RestControlPoint.X;
+				ControlPoint += SampleDeformations[SampleCount + 2] * RestControlPoint.Y;
+				ControlPoint += SampleDeformations[SampleCount + 3] * RestControlPoint.Z;
+
+				return ControlPoint;
+			};
+
+			const int32 NumVertexIndices = DeformingMesh.VertexCount();
+			const int32 NumTasks = FMath::Max(FMath::Min(FTaskGraphInterface::Get().GetNumWorkerThreads(), NumVertexIndices), 1);
+			constexpr int32 MinVerticesByTask = 10;
+			const int32 VerticesByTask = FMath::Max(FMath::Min(FMath::DivideAndRoundUp(NumVertexIndices, NumTasks), MinVerticesByTask), 1);
+			const int32 NumBatches = FMath::DivideAndRoundUp(NumVertexIndices, VerticesByTask);
+			TArray<UE::Tasks::FTask> PendingTasks;
+			PendingTasks.Reserve(NumBatches);
+
+			for (int32 BatchIndex = 0; BatchIndex < NumBatches; BatchIndex++)
+			{
+				const int32 StartIndex = BatchIndex * VerticesByTask;
+				int32 EndIndex = (BatchIndex + 1) * VerticesByTask;
+				EndIndex = BatchIndex == NumBatches - 1 ? FMath::Min(NumVertexIndices, EndIndex) : EndIndex;
 
-				const FVector3f RestPoint(FVector3f(DeformingMesh.GetVertexRef(VertexIndex)));
-				const FVector3f DeformedPoint = DeformPoint(RestPoint);
+				UE::Tasks::FTask PendingTask = UE::Tasks::Launch(UE_SOURCE_LOCATION, [&DeformingMesh, StartIndex, EndIndex, BatchIndex, DeformPoint, &NormalOverlays]
+				{
+					for (int32 VertexIndex = StartIndex; VertexIndex < EndIndex; VertexIndex++)
+					{
+						check(DeformingMesh.GetVerticesRefCounts().IsValid(VertexIndex));
+						const FVector3f RestPoint(FVector3f(DeformingMesh.GetVertexRef(VertexIndex)));
+						const FVector3f DeformedPoint = DeformPoint(RestPoint);
 
-				DeformingMesh.SetVertex(VertexIndex, FVector3d(DeformedPoint));
+						DeformingMesh.SetVertex(VertexIndex, FVector3d(DeformedPoint));
 
-				for (UE::Geometry::FDynamicMeshNormalOverlay* const Overlay : NormalOverlays)
-				{
-					Overlay->EnumerateVertexElements(VertexIndex,
-						[&RestPoint, &DeformedPoint, &DeformPoint, Overlay](int32 TriangleID, int32 ElementID, const FVector3f& Normal)
+						for (UE::Geometry::FDynamicMeshNormalOverlay* const Overlay : NormalOverlays)
 						{
-							const FVector3f RestNormalEnd = RestPoint + Normal;
-							const FVector3f DeformedNormalEnd = DeformPoint(RestNormalEnd);
-							const FVector3f DeformedNorml = (DeformedNormalEnd - DeformedPoint).GetSafeNormal();
-							Overlay->SetElement(ElementID, DeformedNorml);
-							return true;
+							Overlay->EnumerateVertexElements(VertexIndex,
+								[&RestPoint, &DeformedPoint, &DeformPoint, Overlay](int32 TriangleID, int32 ElementID, const FVector3f& Normal)
+								{
+									const FVector3f RestNormalEnd = RestPoint + Normal;
+									const FVector3f DeformedNormalEnd = DeformPoint(RestNormalEnd);
+									const FVector3f DeformedNorml = (DeformedNormalEnd - DeformedPoint).GetSafeNormal();
+									Overlay->SetElement(ElementID, DeformedNorml);
+									return true;
+								}
+							);
 						}
-					);
-				}
+					}
+				});
+				PendingTasks.Add(PendingTask);
 			}
-		}		
+			UE::Tasks::Wait(PendingTasks);
+		}
 	}
 
 	void FRBFInterpolation::GenerateWeights(const UE::Geometry::FDynamicMesh3& BaseMesh, int32 NumInterpolationPoints, FMeshResizingRBFInterpolationData& OutData)
