Optimize Skeletal Mesh To Mesh Dataflow Node
Transfer Normal Overlay has been multithreaded and is 70 times faster now

diff --git a/Engine/Plugins/Experimental/MeshResizing/Source/MeshResizingCore/Private/MeshResizing/BaseBodyTools.cpp b/Engine/Plugins/Experimental/MeshResizing/Source/MeshResizingCore/Private/MeshResizing/BaseBodyTools.cpp
--- a/Engine/Plugins/Experimental/MeshResizing/Source/MeshResizingCore/Private/MeshResizing/BaseBodyTools.cpp
+++ b/Engine/Plugins/Experimental/MeshResizing/Source/MeshResizingCore/Private/MeshResizing/BaseBodyTools.cpp
@@ -4,6 +4,8 @@
 #include "DynamicMesh/DynamicMesh3.h"
 #include "DynamicMesh/DynamicMeshAttributeSet.h"
 #include "DynamicMesh/DynamicVertexAttribute.h"
+#include "Async/TaskGraphInterfaces.h"
+#include "Tasks/Task.h"
 
 namespace UE::MeshResizing
 {
@@ -76,23 +78,43 @@ namespace UE::MeshResizing
 
 		static void TransferNormalOverlayValues(const UE::Geometry::FDynamicMeshNormalOverlay& TargetOverlay, UE::Geometry::FDynamicMeshNormalOverlay& ProxyOverlay, const FVertexMappingAttribute& SourceAttr, const TMap<int32, int32>& MappedTargetToMeshTarget)
 		{
-			TArray<int32> TargetVertexElements; // Reused array
-			for (int32 ProxyElementID : ProxyOverlay.ElementIndicesItr())
+			const int32 NumElements = ProxyOverlay.MaxElementID();
+			const int32 NumTasks = FMath::Max(FMath::Min(FTaskGraphInterface::Get().GetNumWorkerThreads(), NumElements), 1);
+			constexpr int32 MinVerticesByTask = 20;
+			const int32 VerticesByTask = FMath::Max(FMath::Max(FMath::DivideAndRoundUp(NumElements, NumTasks), MinVerticesByTask), 1);
+			const int32 NumBatches = FMath::DivideAndRoundUp(NumElements, VerticesByTask);
+			TArray<UE::Tasks::FTask> PendingTasks;
+			PendingTasks.Reserve(NumBatches);
+			for (int32 BatchIndex = 0; BatchIndex < NumBatches; BatchIndex++)
 			{
-				const int32 ProxyVID = ProxyOverlay.GetParentVertex(ProxyElementID);
-				const int32 MappedVID = SourceAttr.GetMappedValue(ProxyVID);
-				if (const int32* const TargetVID = MappedTargetToMeshTarget.Find(MappedVID))
+				const int32 StartIndex = BatchIndex * VerticesByTask;
+				int32 EndIndex = (BatchIndex + 1) * VerticesByTask;
+				EndIndex = BatchIndex == NumBatches - 1 ? FMath::Min(NumElements, EndIndex) : EndIndex;
+				UE::Tasks::FTask PendingTask = UE::Tasks::Launch(UE_SOURCE_LOCATION, [&ProxyOverlay, &SourceAttr, &TargetOverlay, &MappedTargetToMeshTarget, StartIndex, EndIndex]()
 				{
-					TargetOverlay.GetVertexElements(*TargetVID, TargetVertexElements);
-					if (ensure(TargetVertexElements.Num() > 0))
+					TArray<int32> TargetVertexElements; // Reused array
+					for (int32 ProxyElementID = StartIndex; ProxyElementID < EndIndex; ProxyElementID++)
 					{
-						ensure(TargetVertexElements.Num() == 1);
-						ProxyOverlay.SetElement(ProxyElementID, TargetOverlay.GetElement(TargetVertexElements[0]));
+						if (ProxyOverlay.IsElement(ProxyElementID))
+						{
+							const int32 ProxyVID = ProxyOverlay.GetParentVertex(ProxyElementID);
+							const int32 MappedVID = SourceAttr.GetMappedValue(ProxyVID);
+							if (const int32* const TargetVID = MappedTargetToMeshTarget.Find(MappedVID))
+							{
+								TargetOverlay.GetVertexElements(*TargetVID, TargetVertexElements);
+								if (ensure(TargetVertexElements.Num() > 0))
+								{
+									ensure(TargetVertexElements.Num() == 1);
+									ProxyOverlay.SetElement(ProxyElementID, TargetOverlay.GetElement(TargetVertexElements[0]));
+								}
+							}
+						}
 					}
-				}
+				});
+				PendingTasks.Add(PendingTask);
 			}
+			UE::Tasks::Wait(PendingTasks);
 		}
-
 		static void InterpolateNormalOverlayValues(const UE::Geometry::FDynamicMeshNormalOverlay& SourceOverlay, const UE::Geometry::FDynamicMeshNormalOverlay& TargetOverlay, float BlendAlpha, UE::Geometry::FDynamicMeshNormalOverlay& ProxyOverlay)
 		{
 			check(SourceOverlay.ElementCount() == TargetOverlay.ElementCount());
