Clean up in resizing mutlithreading code
- Rename variable
- fix min vertex by batch

diff --git a/Engine/Plugins/Experimental/MeshResizing/Source/MeshResizingCore/Private/MeshResizing/BaseBodyTools.cpp b/Engine/Plugins/Experimental/MeshResizing/Source/MeshResizingCore/Private/MeshResizing/BaseBodyTools.cpp
--- a/Engine/Plugins/Experimental/MeshResizing/Source/MeshResizingCore/Private/MeshResizing/BaseBodyTools.cpp
+++ b/Engine/Plugins/Experimental/MeshResizing/Source/MeshResizingCore/Private/MeshResizing/BaseBodyTools.cpp
@@ -80,15 +80,15 @@ namespace UE::MeshResizing
 		{
 			const int32 NumElements = ProxyOverlay.MaxElementID();
 			const int32 NumTasks = FMath::Max(FMath::Min(FTaskGraphInterface::Get().GetNumWorkerThreads(), NumElements), 1);
-			constexpr int32 MinVerticesByTask = 20;
-			const int32 VerticesByTask = FMath::Max(FMath::Max(FMath::DivideAndRoundUp(NumElements, NumTasks), MinVerticesByTask), 1);
-			const int32 NumBatches = FMath::DivideAndRoundUp(NumElements, VerticesByTask);
+			constexpr int32 MinElementsByTask = 20;
+			const int32 ElementsByTask = FMath::Max(FMath::DivideAndRoundUp(NumElements, NumTasks), MinElementsByTask);
+			const int32 NumBatches = FMath::DivideAndRoundUp(NumElements, ElementsByTask);
 			TArray<UE::Tasks::FTask> PendingTasks;
 			PendingTasks.Reserve(NumBatches);
 			for (int32 BatchIndex = 0; BatchIndex < NumBatches; BatchIndex++)
 			{
-				const int32 StartIndex = BatchIndex * VerticesByTask;
-				int32 EndIndex = (BatchIndex + 1) * VerticesByTask;
+				const int32 StartIndex = BatchIndex * ElementsByTask;
+				int32 EndIndex = (BatchIndex + 1) * ElementsByTask;
 				EndIndex = BatchIndex == NumBatches - 1 ? FMath::Min(NumElements, EndIndex) : EndIndex;
 				UE::Tasks::FTask PendingTask = UE::Tasks::Launch(UE_SOURCE_LOCATION, [&ProxyOverlay, &SourceAttr, &TargetOverlay, &MappedTargetToMeshTarget, StartIndex, EndIndex]()
 				{
diff --git a/Engine/Plugins/Experimental/MeshResizing/Source/MeshResizingCore/Private/MeshResizing/RBFInterpolation.cpp b/Engine/Plugins/Experimental/MeshResizing/Source/MeshResizingCore/Private/MeshResizing/RBFInterpolation.cpp
--- a/Engine/Plugins/Experimental/MeshResizing/Source/MeshResizingCore/Private/MeshResizing/RBFInterpolation.cpp
+++ b/Engine/Plugins/Experimental/MeshResizing/Source/MeshResizingCore/Private/MeshResizing/RBFInterpolation.cpp
@@ -32,7 +32,7 @@ namespace UE::MeshResizing
 
 			const uint32 NumTasks = FMath::Max(FMath::Min(uint32(FTaskGraphInterface::Get().GetNumWorkerThreads()), EntryCount), 1u);
 			constexpr uint32 MinEntryByTask = 10;
-			const uint32 EntryByTask = FMath::Max(FMath::Min(FMath::DivideAndRoundUp(EntryCount, NumTasks), MinEntryByTask), 1u);
+			const uint32 EntryByTask = FMath::Max(FMath::DivideAndRoundUp(EntryCount, NumTasks), MinEntryByTask);
 			const uint32 NumBatches = FMath::DivideAndRoundUp(EntryCount, EntryByTask);
 			TArray<UE::Tasks::FTask> PendingTasks;
 			PendingTasks.Reserve(NumBatches);
@@ -100,7 +100,7 @@ namespace UE::MeshResizing
 			};
 
 			const int32 NumVertexIndices = DeformingMesh.VertexCount();
-			const int32 NumTasks = FMath::Max(FMath::Min(FTaskGraphInterface::Get().GetNumWorkerThreads(), NumVertexIndices), 1);
+			const int32 NumTasks = FMath::Max(FTaskGraphInterface::Get().GetNumWorkerThreads(), NumVertexIndices);
 			constexpr int32 MinVerticesByTask = 10;
 			const int32 VerticesByTask = FMath::Max(FMath::Min(FMath::DivideAndRoundUp(NumVertexIndices, NumTasks), MinVerticesByTask), 1);
 			const int32 NumBatches = FMath::DivideAndRoundUp(NumVertexIndices, VerticesByTask);
