Follow-up fixes from code changes to support MakeResizeableOutfit and ResizeOutfit templates.

diff --git a/Engine/Plugins/Experimental/ChaosOutfitAsset/Source/ChaosOutfitAssetDataflowNodes/Private/ChaosOutfitAsset/FilterSizedOutfitNode.cpp b/Engine/Plugins/Experimental/ChaosOutfitAsset/Source/ChaosOutfitAssetDataflowNodes/Private/ChaosOutfitAsset/FilterSizedOutfitNode.cpp
--- a/Engine/Plugins/Experimental/ChaosOutfitAsset/Source/ChaosOutfitAssetDataflowNodes/Private/ChaosOutfitAsset/FilterSizedOutfitNode.cpp
+++ b/Engine/Plugins/Experimental/ChaosOutfitAsset/Source/ChaosOutfitAssetDataflowNodes/Private/ChaosOutfitAsset/FilterSizedOutfitNode.cpp
@@ -34,8 +34,11 @@ void FChaosOutfitAssetFilterSizedOutfitNode::Evaluate(UE::Dataflow::FContext& Co
 				FString InSizeName = GetValue(Context, &SizeName);
 				if (!InSizeName.IsEmpty() && !InOutfit->HasBodySize(InSizeName))
 				{
-					// Only show a warning when the size name is non empty
-					Context.Warning(FString::Printf(TEXT("The given body size [%s] doesn't exist in the input Outfit."), *InSizeName), this, Out);
+					if (InSizeName != FName(NAME_None).ToString())
+					{
+						// Only show a warning when the size name is non empty/ not from a NAME_None FName
+						Context.Warning(FString::Printf(TEXT("The given body size [%s] doesn't exist in the input Outfit."), *InSizeName), this, Out);
+					}
 					InSizeName.Reset();
 				}
 
diff --git a/Engine/Plugins/Experimental/ChaosOutfitAsset/Source/ChaosOutfitAssetEngine/Public/ChaosOutfitAsset/CollectionOutfitFacade.h b/Engine/Plugins/Experimental/ChaosOutfitAsset/Source/ChaosOutfitAssetEngine/Public/ChaosOutfitAsset/CollectionOutfitFacade.h
--- a/Engine/Plugins/Experimental/ChaosOutfitAsset/Source/ChaosOutfitAssetEngine/Public/ChaosOutfitAsset/CollectionOutfitFacade.h
+++ b/Engine/Plugins/Experimental/ChaosOutfitAsset/Source/ChaosOutfitAssetEngine/Public/ChaosOutfitAsset/CollectionOutfitFacade.h
@@ -12,7 +12,6 @@ template<typename T> class TManagedArray;
 class UChaosClothAssetBase;
 class USkeletalMesh;
 struct FManagedArrayCollection;
-struct FMeshResizingRBFInterpolationData;
 
 namespace UE::Chaos::OutfitAsset
 {
diff --git a/Engine/Plugins/Experimental/ChaosOutfitAsset/Source/ChaosOutfitAssetEngine/Public/ChaosOutfitAsset/Outfit.h b/Engine/Plugins/Experimental/ChaosOutfitAsset/Source/ChaosOutfitAssetEngine/Public/ChaosOutfitAsset/Outfit.h
--- a/Engine/Plugins/Experimental/ChaosOutfitAsset/Source/ChaosOutfitAssetEngine/Public/ChaosOutfitAsset/Outfit.h
+++ b/Engine/Plugins/Experimental/ChaosOutfitAsset/Source/ChaosOutfitAssetEngine/Public/ChaosOutfitAsset/Outfit.h
@@ -107,7 +107,7 @@ public:
 		return Pieces;
 	}
 
-	TArray<FChaosOutfitPiece>& GetPieces()
+	TArrayView<FChaosOutfitPiece> GetPieces()
 	{
 		return Pieces;
 	}
diff --git a/Engine/Plugins/Experimental/ChaosOutfitAsset/Source/ChaosOutfitAssetEngine/Public/ChaosOutfitAsset/SizedOutfitSource.h b/Engine/Plugins/Experimental/ChaosOutfitAsset/Source/ChaosOutfitAssetEngine/Public/ChaosOutfitAsset/SizedOutfitSource.h
--- a/Engine/Plugins/Experimental/ChaosOutfitAsset/Source/ChaosOutfitAssetEngine/Public/ChaosOutfitAsset/SizedOutfitSource.h
+++ b/Engine/Plugins/Experimental/ChaosOutfitAsset/Source/ChaosOutfitAssetEngine/Public/ChaosOutfitAsset/SizedOutfitSource.h
@@ -43,8 +43,8 @@ struct FChaosSizedOutfitSource final
 	/**
 	 * The number of interpolation points used in the resizing algorithm. These points are distributed evenly across the entire body.
 	 * Increasing this number increases the quality of the resizing operation, but at additional cost, including the initial operation of 
-	 *  generating your resizable outfit, and the size of your Outfit on disk.
-	 * If you find you cannot generate acceptable resizing results by increasing this number, we recommend generating a new Size.
+	 *  generating your resizable Outfit, and the size of your Outfit on disk.
+	 * If you find you cannot generate acceptable resizing results by increasing this number, we recommend adding a new Size.
 	 */
 	UPROPERTY(EditAnywhere, AdvancedDisplay, Category = "OutfitBodySize", Meta = (ClampMin = 0, ClampMax = 5000))
 	int32 NumResizingInterpolationPoints = UE::Chaos::OutfitAsset::DefaultNumRBFInterpolationPoints;
diff --git a/Engine/Plugins/Experimental/MeshResizing/Source/MeshResizingNodes/Private/MeshResizing/BaseBodyDataflowNodes.h b/Engine/Plugins/Experimental/MeshResizing/Source/MeshResizingNodes/Private/MeshResizing/BaseBodyDataflowNodes.h
--- a/Engine/Plugins/Experimental/MeshResizing/Source/MeshResizingNodes/Private/MeshResizing/BaseBodyDataflowNodes.h
+++ b/Engine/Plugins/Experimental/MeshResizing/Source/MeshResizingNodes/Private/MeshResizing/BaseBodyDataflowNodes.h
@@ -55,12 +55,13 @@ private:
 	int32 LODLevel = 0;
 
 #if WITH_EDITORONLY_DATA
-	/** Record ImportedVertices (if available) as NonManifold mapping data. This requires Editor-Only data.*/
+
+	/** Generate from the SkeletalMeshLODModel (vertex order will match SKM vertex order). Record ImportedVertices (if available) as NonManifold mapping data. This requires Editor-Only data.*/
 	UPROPERTY(EditAnywhere, Category = "SkeletalMesh")
 	bool bRecordImportedVertices = true;
 
-	/** Generate from mesh description. Requires Editor-Only data.*/
-	UPROPERTY(EditAnywhere, Category = "SkeletalMesh")
+	/** Generate from mesh description (vertex order will match mesh description / ImportedVertices). Requires Editor-Only data.*/
+	UPROPERTY(EditAnywhere, Category = "SkeletalMesh", meta = (EditCondition = "!bRecordImportedVertices"))
 	bool bUseMeshDescription = false;
 #endif
 };
diff --git a/Engine/Plugins/Experimental/MeshResizing/Source/MeshResizingNodes/Private/MeshResizing/RBFInterpolationNodes.cpp b/Engine/Plugins/Experimental/MeshResizing/Source/MeshResizingNodes/Private/MeshResizing/RBFInterpolationNodes.cpp
--- a/Engine/Plugins/Experimental/MeshResizing/Source/MeshResizingNodes/Private/MeshResizing/RBFInterpolationNodes.cpp
+++ b/Engine/Plugins/Experimental/MeshResizing/Source/MeshResizingNodes/Private/MeshResizing/RBFInterpolationNodes.cpp
@@ -77,10 +77,8 @@ FApplyRBFResizingNode::FApplyRBFResizingNode(const UE::Dataflow::FNodeParameters
 	FDataflowNode(InParam, InGuid)
 {
 	RegisterInputConnection(&MeshToResize);
-#if WITH_EDITORONLY_DATA
 	RegisterInputConnection(&TargetSkeletalMesh);
 	RegisterInputConnection(&TargetSkeletalMeshLODIndex);
-#endif
 	RegisterInputConnection(&TargetMesh);
 	RegisterInputConnection(&InterpolationData);
 
@@ -90,25 +88,59 @@ FApplyRBFResizingNode::FApplyRBFResizingNode(const UE::Dataflow::FNodeParameters
 
 void FApplyRBFResizingNode::Evaluate(UE::Dataflow::FContext& Context, const FDataflowOutput* Out) const
 {
-	if (TObjectPtr<const UDataflowMesh> InMeshToResize = GetValue(Context, &MeshToResize))
+	if (Out->IsA<TObjectPtr<const UDataflowMesh>>(&ResizedMesh))
 	{
-#if WITH_EDITORONLY_DATA
-		if (bUseSkeletalMeshTarget)
+		if (TObjectPtr<const UDataflowMesh> InMeshToResize = GetValue(Context, &MeshToResize))
 		{
-			if (TObjectPtr<const USkeletalMesh> InTargetSKM = GetValue(Context, &TargetSkeletalMesh))
+#if WITH_EDITORONLY_DATA
+			if (bUseSkeletalMeshTarget)
+			{
+				if (TObjectPtr<const USkeletalMesh> InTargetSKM = GetValue(Context, &TargetSkeletalMesh))
+				{
+					const int32 InLODIndex = GetValue(Context, &TargetSkeletalMeshLODIndex);
+					if (const FMeshDescription* const MeshDescription = InTargetSKM->GetMeshDescription(InLODIndex))
+					{
+						const FMeshResizingRBFInterpolationData& InInterpolationData = GetValue(Context, &InterpolationData);
+						if (InMeshToResize->GetDynamicMesh() && InInterpolationData.SampleIndices.Num())
+						{
+							TObjectPtr<UDataflowMesh> OutResizedMesh = NewObject<UDataflowMesh>();
+							FDynamicMesh3 ResizeMesh;
+							ResizeMesh.Copy(InMeshToResize->GetDynamicMeshRef());
+
+							// Do the interpolation 
+							UE::MeshResizing::FRBFInterpolation::DeformPoints(*MeshDescription, InInterpolationData, bInterpolateNormals, ResizeMesh);
+							OutResizedMesh->SetDynamicMesh(MoveTemp(ResizeMesh));
+							OutResizedMesh->SetMaterials(InMeshToResize->GetMaterials());
+							SetValue<TObjectPtr<const UDataflowMesh>>(Context, OutResizedMesh, &ResizedMesh);
+							return;
+						}
+					}
+				}
+			}
+			else
+#endif
 			{
-				const int32 InLODIndex = GetValue(Context, &TargetSkeletalMeshLODIndex);
-				if (const FMeshDescription* const MeshDescription = InTargetSKM->GetMeshDescription(InLODIndex))
+				if (TObjectPtr<const UDataflowMesh> InTargetMesh = GetValue(Context, &TargetMesh))
 				{
 					const FMeshResizingRBFInterpolationData& InInterpolationData = GetValue(Context, &InterpolationData);
-					if (InMeshToResize->GetDynamicMesh() && InInterpolationData.SampleIndices.Num())
+
+					if (InMeshToResize->GetDynamicMesh() && InTargetMesh->GetDynamicMesh() && InInterpolationData.SampleIndices.Num())
 					{
 						TObjectPtr<UDataflowMesh> OutResizedMesh = NewObject<UDataflowMesh>();
 						FDynamicMesh3 ResizeMesh;
 						ResizeMesh.Copy(InMeshToResize->GetDynamicMeshRef());
 
-						// Do the interpolation 
-						UE::MeshResizing::FRBFInterpolation::DeformPoints(*MeshDescription, InInterpolationData, bInterpolateNormals, ResizeMesh);
+						const FDynamicMesh3& TargetMeshRef = InTargetMesh->GetDynamicMeshRef();
+
+						// Get target mesh samples
+						TArray<FVector3f> TargetSamplePoints;
+						TargetSamplePoints.Reserve(InInterpolationData.SampleIndices.Num());
+						for (const uint32 SampleIndex : InInterpolationData.SampleIndices)
+						{
+							TargetSamplePoints.Emplace(FVector3f(TargetMeshRef.GetVertex(SampleIndex)));
+						}
+						// Do the interpolation
+						UE::MeshResizing::FRBFInterpolation::DeformPoints(TargetSamplePoints, InInterpolationData, bInterpolateNormals, ResizeMesh);
 						OutResizedMesh->SetDynamicMesh(MoveTemp(ResizeMesh));
 						OutResizedMesh->SetMaterials(InMeshToResize->GetMaterials());
 						SetValue<TObjectPtr<const UDataflowMesh>>(Context, OutResizedMesh, &ResizedMesh);
@@ -117,38 +149,7 @@ void FApplyRBFResizingNode::Evaluate(UE::Dataflow::FContext& Context, const FDat
 				}
 			}
 		}
-		else
-#endif
-		{
-			if (TObjectPtr<const UDataflowMesh> InTargetMesh = GetValue(Context, &TargetMesh))
-			{
-				const FMeshResizingRBFInterpolationData& InInterpolationData = GetValue(Context, &InterpolationData);
-
-				if (InMeshToResize->GetDynamicMesh() && InTargetMesh->GetDynamicMesh() && InInterpolationData.SampleIndices.Num())
-				{
-					TObjectPtr<UDataflowMesh> OutResizedMesh = NewObject<UDataflowMesh>();
-					FDynamicMesh3 ResizeMesh;
-					ResizeMesh.Copy(InMeshToResize->GetDynamicMeshRef());
-
-					const FDynamicMesh3& TargetMeshRef = InTargetMesh->GetDynamicMeshRef();
 
-					// Get target mesh samples
-					TArray<FVector3f> TargetSamplePoints;
-					TargetSamplePoints.Reserve(InInterpolationData.SampleIndices.Num());
-					for (const uint32 SampleIndex : InInterpolationData.SampleIndices)
-					{
-						TargetSamplePoints.Emplace(FVector3f(TargetMeshRef.GetVertex(SampleIndex)));
-					}
-					// Do the interpolation
-					UE::MeshResizing::FRBFInterpolation::DeformPoints(TargetSamplePoints, InInterpolationData, bInterpolateNormals, ResizeMesh);
-					OutResizedMesh->SetDynamicMesh(MoveTemp(ResizeMesh));
-					OutResizedMesh->SetMaterials(InMeshToResize->GetMaterials());
-					SetValue<TObjectPtr<const UDataflowMesh>>(Context, OutResizedMesh, &ResizedMesh);
-					return;
-				}
-			}
-		}
+		SafeForwardInput(Context, &MeshToResize, &ResizedMesh);
 	}
-
-	SafeForwardInput(Context, &MeshToResize, &ResizedMesh);
 }
\ No newline at end of file
diff --git a/Engine/Plugins/Experimental/MeshResizing/Source/MeshResizingNodes/Private/MeshResizing/RBFInterpolationNodes.h b/Engine/Plugins/Experimental/MeshResizing/Source/MeshResizingNodes/Private/MeshResizing/RBFInterpolationNodes.h
--- a/Engine/Plugins/Experimental/MeshResizing/Source/MeshResizingNodes/Private/MeshResizing/RBFInterpolationNodes.h
+++ b/Engine/Plugins/Experimental/MeshResizing/Source/MeshResizingNodes/Private/MeshResizing/RBFInterpolationNodes.h
@@ -80,18 +80,17 @@ private:
 	/** Use a skeletal mesh for the target mesh (instead of a dynamic mesh)*/
 	UPROPERTY(EditAnywhere, Category = "RBF Resize")
 	bool bUseSkeletalMeshTarget = false;
+#else
+	static constexpr bool bUseSkeletalMeshTarget = false;
+#endif
 
 	/** The target mesh that corresponds with the SourceMesh used to generate the InterpolationData. Must have matching vertices with SourceMesh */
-	UPROPERTY(meta = (DataflowInput, EditCondition = "UseSkeletalMeshTarget"))
+	UPROPERTY(meta = (DataflowInput, EditCondition = "bUseSkeletalMeshTarget"))
 	TObjectPtr<const USkeletalMesh> TargetSkeletalMesh;
 
-	UPROPERTY(EditAnywhere, Category = "RBF Resize", meta = (DataflowInput, EditCondition = "UseSkeletalMeshTarget", ClampMin = 0))
+	UPROPERTY(EditAnywhere, Category = "RBF Resize", meta = (DataflowInput, EditCondition = "bUseSkeletalMeshTarget", ClampMin = 0))
 	int32 TargetSkeletalMeshLODIndex = 0;
 
-#else
-	static constexpr bool bUseSkeletalMeshTarget = false;
-#endif
-
 	/** The target mesh that corresponds with the SourceMesh used to generate the InterpolationData. Must have matching vertices with SourceMesh */
 	UPROPERTY(meta = (DataflowInput, EditCondition = "!bUseSkeletalMeshTarget"))
 	TObjectPtr<const UDataflowMesh> TargetMesh;
