Add Chaos Cloth ApplyResizing node to directly resize the cloth collection.

#virtualized

diff --git a/Engine/Plugins/ChaosClothAssetEditor/ChaosClothAssetEditor.uplugin b/Engine/Plugins/ChaosClothAssetEditor/ChaosClothAssetEditor.uplugin
--- a/Engine/Plugins/ChaosClothAssetEditor/ChaosClothAssetEditor.uplugin
+++ b/Engine/Plugins/ChaosClothAssetEditor/ChaosClothAssetEditor.uplugin
@@ -78,6 +78,11 @@
 				"Mac",
 				"Linux"
 			]
+		},
+
+		{
+			"Name": "MeshResizing",
+			"Enabled": true
 		}
 	]
 }
diff --git a/Engine/Plugins/ChaosClothAssetEditor/Source/ChaosClothAssetDataflowNodes/ChaosClothAssetDataflowNodes.Build.cs b/Engine/Plugins/ChaosClothAssetEditor/Source/ChaosClothAssetDataflowNodes/ChaosClothAssetDataflowNodes.Build.cs
--- a/Engine/Plugins/ChaosClothAssetEditor/Source/ChaosClothAssetDataflowNodes/ChaosClothAssetDataflowNodes.Build.cs
+++ b/Engine/Plugins/ChaosClothAssetEditor/Source/ChaosClothAssetDataflowNodes/ChaosClothAssetDataflowNodes.Build.cs
@@ -10,6 +10,7 @@ public class ChaosClothAssetDataflowNodes : ModuleRules
 			new string[]
 			{
 				"Core",
+				"MeshResizingCore",
 				// ... add other public dependencies that you statically link with here ...
 			}
 		);
diff --git a/Engine/Plugins/ChaosClothAssetEditor/Source/ChaosClothAssetDataflowNodes/Private/ChaosClothAsset/ApplyResizingNode.cpp b/Engine/Plugins/ChaosClothAssetEditor/Source/ChaosClothAssetDataflowNodes/Private/ChaosClothAsset/ApplyResizingNode.cpp
new file mode 100644
--- /dev/null
+++ b/Engine/Plugins/ChaosClothAssetEditor/Source/ChaosClothAssetDataflowNodes/Private/ChaosClothAsset/ApplyResizingNode.cpp
@@ -0,0 +1,66 @@
+// Copyright Epic Games, Inc. All Rights Reserved.
+
+#include "ChaosClothAsset/ApplyResizingNode.h"
+#include "ChaosClothAsset/CollectionClothFacade.h"
+#include "Dataflow/DataflowInputOutput.h"
+#include "Engine/SkeletalMesh.h"
+
+#include UE_INLINE_GENERATED_CPP_BY_NAME(ApplyResizingNode)
+
+#define LOCTEXT_NAMESPACE "ChaosClothAssetApplyResizingNode"
+
+FChaosClothAssetApplyResizingNode::FChaosClothAssetApplyResizingNode(const UE::Dataflow::FNodeParameters& InParam, FGuid InGuid)
+	: FDataflowNode(InParam, InGuid)
+{
+	RegisterInputConnection(&Collection);
+	RegisterOutputConnection(&Collection, &Collection);
+	RegisterInputConnection(&TargetSkeletalMesh);
+	RegisterInputConnection(&InterpolationData);
+	RegisterInputConnection(&SkeletalMeshLODIndex)
+		.SetCanHidePin(true)
+		.SetPinIsHidden(true);
+	RegisterInputConnection(&bForceApplyToRenderMesh)
+		.SetCanHidePin(true)
+		.SetPinIsHidden(true);
+}
+
+void FChaosClothAssetApplyResizingNode::Evaluate(UE::Dataflow::FContext& Context, const FDataflowOutput* Out) const
+{
+	if (Out->IsA<FManagedArrayCollection>(&Collection))
+	{
+		using namespace UE::Chaos::ClothAsset;
+		// Evaluate in collection
+		FManagedArrayCollection InCollection = GetValue<FManagedArrayCollection>(Context, &Collection);
+		const TSharedRef<FManagedArrayCollection> ClothCollection = MakeShared<FManagedArrayCollection>(MoveTemp(InCollection));
+		FCollectionClothFacade ClothFacade(ClothCollection);
+		if (ClothFacade.IsValid())
+		{
+			if (TObjectPtr<const USkeletalMesh> InTargetSKM = GetValue(Context, &TargetSkeletalMesh))
+			{
+				const int32 InLODIndex = GetValue(Context, &SkeletalMeshLODIndex);
+				if (InTargetSKM->IsValidLODIndex(InLODIndex) && InTargetSKM->HasMeshDescription(InLODIndex))
+				{
+					if (const FMeshDescription* const MeshDescription = InTargetSKM->GetMeshDescription(InLODIndex))
+					{
+						const FMeshResizingRBFInterpolationData& InInterpolationData = GetValue(Context, &InterpolationData);
+						if (InInterpolationData.SampleIndices.Num())
+						{
+							if (!GetValue(Context, &bForceApplyToRenderMesh) && ClothFacade.HasValidSimulationData())
+							{
+								UE::MeshResizing::FRBFInterpolation::DeformPoints(*MeshDescription, InInterpolationData, ClothFacade.GetSimPosition3D(), ClothFacade.GetSimNormal());
+							}
+							else if (ClothFacade.HasValidRenderData())
+							{
+								UE::MeshResizing::FRBFInterpolation::DeformPoints(*MeshDescription, InInterpolationData, ClothFacade.GetRenderPosition(), ClothFacade.GetRenderNormal(), ClothFacade.GetRenderTangentU(), ClothFacade.GetRenderTangentV());
+							}
+						}
+					}
+				}
+			}
+		}
+
+		SetValue(Context, MoveTemp(*ClothCollection), &Collection);
+	}
+}
+
+#undef LOCTEXT_NAMESPACE
diff --git a/Engine/Plugins/ChaosClothAssetEditor/Source/ChaosClothAssetDataflowNodes/Private/ChaosClothAsset/ChaosClothAssetDataflowNodesModule.cpp b/Engine/Plugins/ChaosClothAssetEditor/Source/ChaosClothAssetDataflowNodes/Private/ChaosClothAsset/ChaosClothAssetDataflowNodesModule.cpp
--- a/Engine/Plugins/ChaosClothAssetEditor/Source/ChaosClothAssetDataflowNodes/Private/ChaosClothAsset/ChaosClothAssetDataflowNodesModule.cpp
+++ b/Engine/Plugins/ChaosClothAssetEditor/Source/ChaosClothAssetDataflowNodes/Private/ChaosClothAsset/ChaosClothAssetDataflowNodesModule.cpp
@@ -3,6 +3,7 @@
 #include "ChaosClothAsset/AddStitchNode.h"
 #include "ChaosClothAsset/AddWeightMapNode.h"
 #include "ChaosClothAsset/ApplyProxyDeformerNode.h"
+#include "ChaosClothAsset/ApplyResizingNode.h"
 #include "ChaosClothAsset/AttributeNode.h"
 #include "ChaosClothAsset/BindToRootBoneNode.h"
 #include "ChaosClothAsset/BlendVerticesNode.h"
@@ -99,6 +100,7 @@ namespace UE::Chaos::ClothAsset
 			DATAFLOW_NODE_REGISTER_CREATION_FACTORY(FChaosClothAssetWeightMapNode);
 			DATAFLOW_NODE_REGISTER_CREATION_FACTORY(FChaosClothAssetAttributeNode_v2);
 			DATAFLOW_NODE_REGISTER_CREATION_FACTORY(FChaosClothAssetApplyProxyDeformerNode);
+			DATAFLOW_NODE_REGISTER_CREATION_FACTORY(FChaosClothAssetApplyResizingNode);
 			DATAFLOW_NODE_REGISTER_CREATION_FACTORY(FChaosClothAssetBindToRootBoneNode);
 			DATAFLOW_NODE_REGISTER_CREATION_FACTORY(FChaosClothAssetBlendVerticesNode);
 			DATAFLOW_NODE_REGISTER_CREATION_FACTORY(FChaosClothAssetCollectionToDynamicMeshNode);
diff --git a/Engine/Plugins/ChaosClothAssetEditor/Source/ChaosClothAssetDataflowNodes/Public/ChaosClothAsset/ApplyResizingNode.h b/Engine/Plugins/ChaosClothAssetEditor/Source/ChaosClothAssetDataflowNodes/Public/ChaosClothAsset/ApplyResizingNode.h
new file mode 100644
--- /dev/null
+++ b/Engine/Plugins/ChaosClothAssetEditor/Source/ChaosClothAssetDataflowNodes/Public/ChaosClothAsset/ApplyResizingNode.h
@@ -0,0 +1,46 @@
+// Copyright Epic Games, Inc. All Rights Reserved.
+
+#pragma once 
+
+#include "Dataflow/DataflowNode.h"
+#include "GeometryCollection/ManagedArrayCollection.h"
+#include "MeshResizing/RBFInterpolation.h"
+#include "ApplyResizingNode.generated.h"
+
+class USkeletalMesh;
+
+/** Apply resizing for a given Target Mesh.*/
+USTRUCT(Meta = (DataflowCloth, Experimental))
+struct FChaosClothAssetApplyResizingNode final : public FDataflowNode
+{
+	GENERATED_USTRUCT_BODY()
+	DATAFLOW_NODE_DEFINE_INTERNAL(FChaosClothAssetApplyResizingNode, "ApplyResizing", "Cloth", "Apply Cloth Asset Resizing")
+
+public:
+
+	FChaosClothAssetApplyResizingNode(const UE::Dataflow::FNodeParameters& InParam, FGuid InGuid = FGuid::NewGuid());
+
+private:
+	//~ Begin FDataflowNode interface
+	virtual void Evaluate(UE::Dataflow::FContext& Context, const FDataflowOutput* Out) const override;
+	//~ End FDataflowNode interface
+
+	UPROPERTY(Meta = (Dataflowinput, DataflowOutput, DataflowPassthrough = "Collection"))
+	FManagedArrayCollection Collection;
+
+	/** The target mesh that corresponds with the SourceMesh used to generate the InterpolationData. Must have matching vertices with SourceMesh */
+	UPROPERTY(meta = (DataflowInput))
+	TObjectPtr<const USkeletalMesh> TargetSkeletalMesh;
+
+	/** Which Skeletal Mesh LOD to use for the TargetSkeletalMesh */
+	UPROPERTY(EditAnywhere, Category = "Resizing", meta = (DataflowInput, ClampMin = 0))
+	int32 SkeletalMeshLODIndex = 0;
+
+	/** The pre-calculated base RBF interpolation data.*/
+	UPROPERTY(meta = (DataflowInput))
+	FMeshResizingRBFInterpolationData InterpolationData;
+
+	/** Force apply to the render mesh. Otherwise, the sim mesh will be resized if it exists, only resizing the render mesh if no sim mesh data exists. */
+	UPROPERTY(EditAnywhere, Category = "Resizing", meta = (DataflowInput))
+	bool bForceApplyToRenderMesh = false;
+};
\ No newline at end of file
diff --git a/Engine/Plugins/Experimental/MeshResizing/Source/MeshResizingCore/Private/MeshResizing/RBFInterpolation.cpp b/Engine/Plugins/Experimental/MeshResizing/Source/MeshResizingCore/Private/MeshResizing/RBFInterpolation.cpp
--- a/Engine/Plugins/Experimental/MeshResizing/Source/MeshResizingCore/Private/MeshResizing/RBFInterpolation.cpp
+++ b/Engine/Plugins/Experimental/MeshResizing/Source/MeshResizingCore/Private/MeshResizing/RBFInterpolation.cpp
@@ -25,36 +25,56 @@ namespace UE::MeshResizing
 			OutData.InterpolationWeights = InterpolationWeights.InverseEntries;
 		}
 
+		static FVector3f DeformPoint(const FMeshResizingRBFInterpolationData& InterpolationData, const TArray<FVector3f>& SampleDeformations, const FVector3f& RestControlPoint)
+		{
+			const int32 SampleCount = InterpolationData.SampleRestPositions.Num();
+			check(SampleCount + 4 == SampleDeformations.Num());
+
+			FVector3f ControlPoint = RestControlPoint;
+			for (int32 SampleIndex = 0; SampleIndex < SampleCount; ++SampleIndex)
+			{
+				const FVector3f PositionDelta = RestControlPoint - InterpolationData.SampleRestPositions[SampleIndex];
+				const float FunctionValue = FMath::Sqrt(FVector3f::DotProduct(PositionDelta, PositionDelta) + 1);
+				ControlPoint += FunctionValue * SampleDeformations[SampleIndex];
+			}
+			ControlPoint += SampleDeformations[SampleCount];
+			ControlPoint += SampleDeformations[SampleCount + 1] * RestControlPoint.X;
+			ControlPoint += SampleDeformations[SampleCount + 2] * RestControlPoint.Y;
+			ControlPoint += SampleDeformations[SampleCount + 3] * RestControlPoint.Z;
+
+			return ControlPoint;
+		}
+
 		// From GroomRBFDeformer
 		static void UpdateMeshSamples(const TArray<FVector3f>& TargetPositions, const FMeshResizingRBFInterpolationData& InterpolationData, TArray<FVector3f>& OutSampleDeformations)
 		{
-			const uint32 MaxSampleCount = TargetPositions.Num();
+			const int32 MaxSampleCount = TargetPositions.Num();
 			check(MaxSampleCount == InterpolationData.SampleIndices.Num());
 			check(MaxSampleCount == InterpolationData.SampleRestPositions.Num());
 
-			const uint32 EntryCount = FGroomRBFDeformer::GetEntryCount(MaxSampleCount);
+			const int32 EntryCount = FGroomRBFDeformer::GetEntryCount(MaxSampleCount);
 			OutSampleDeformations.SetNum(EntryCount);
 
-			const uint32 NumTasks = FMath::Max(FMath::Min(uint32(FTaskGraphInterface::Get().GetNumWorkerThreads()), EntryCount), 1u);
-			constexpr uint32 MinEntryByTask = 10;
-			const uint32 EntryByTask = FMath::Max(FMath::DivideAndRoundUp(EntryCount, NumTasks), MinEntryByTask);
-			const uint32 NumBatches = FMath::DivideAndRoundUp(EntryCount, EntryByTask);
+			const int32 NumTasks = FMath::Max(FMath::Min(int32(FTaskGraphInterface::Get().GetNumWorkerThreads()), EntryCount), 1);
+			constexpr int32 MinEntryByTask = 10;
+			const int32 EntryByTask = FMath::Max(FMath::DivideAndRoundUp(EntryCount, NumTasks), MinEntryByTask);
+			const int32 NumBatches = FMath::DivideAndRoundUp(EntryCount, EntryByTask);
 			TArray<UE::Tasks::FTask> PendingTasks;
 			PendingTasks.Reserve(NumBatches);
 
-			for (uint32 BatchIndex = 0; BatchIndex < NumBatches; BatchIndex++)
+			for (int32 BatchIndex = 0; BatchIndex < NumBatches; BatchIndex++)
 			{
-				const uint32 StartIndex = BatchIndex * EntryByTask;
-				uint32 EndIndex = (BatchIndex + 1) * EntryByTask;
+				const int32 StartIndex = BatchIndex * EntryByTask;
+				int32 EndIndex = (BatchIndex + 1) * EntryByTask;
 				EndIndex = BatchIndex == NumBatches - 1 ? FMath::Min(EntryCount, EndIndex) : EndIndex;
 
 				UE::Tasks::FTask PendingTask = UE::Tasks::Launch(UE_SOURCE_LOCATION, [StartIndex, EndIndex, EntryCount, MaxSampleCount, &OutSampleDeformations, &InterpolationData, &TargetPositions]
 				{
-					for (uint32 SampleIndex = StartIndex; SampleIndex < EndIndex; ++SampleIndex)
+					for (int32 SampleIndex = StartIndex; SampleIndex < EndIndex; ++SampleIndex)
 					{
-						uint32 WeightsOffset = SampleIndex * EntryCount;
+						int32 WeightsOffset = SampleIndex * EntryCount;
 						FVector3f SampleDeformation(FVector3f::ZeroVector);
-						for (uint32 Index = 0; Index < MaxSampleCount; ++Index, ++WeightsOffset)
+						for (int32 Index = 0; Index < MaxSampleCount; ++Index, ++WeightsOffset)
 						{
 							SampleDeformation += InterpolationData.InterpolationWeights[WeightsOffset] * (TargetPositions[Index] - InterpolationData.SampleRestPositions[Index]);
 						}
@@ -69,7 +89,7 @@ namespace UE::MeshResizing
 		// From GroomRBFDeformer
 		static void ApplyRBFDeformation(const FMeshResizingRBFInterpolationData& InterpolationData, const TArray<FVector3f>& SampleDeformations, bool bInterpolateNormals, UE::Geometry::FDynamicMesh3& DeformingMesh)
 		{
-			const uint32 SampleCount = InterpolationData.SampleRestPositions.Num();
+			const int32 SampleCount = InterpolationData.SampleRestPositions.Num();
 			check(SampleCount + 4 == SampleDeformations.Num());
 
 			TArray<UE::Geometry::FDynamicMeshNormalOverlay*> NormalOverlays;
@@ -87,21 +107,9 @@ namespace UE::MeshResizing
 				}
 			}
 
-			auto DeformPoint = [SampleCount, &SampleDeformations, &InterpolationData](const FVector3f& RestControlPoint)
+			auto DeformPointLambda = [&SampleDeformations, &InterpolationData](const FVector3f& RestControlPoint)
 			{
-				FVector3f ControlPoint = RestControlPoint;
-				for (uint32 SampleIndex = 0; SampleIndex < SampleCount; ++SampleIndex)
-				{
-					const FVector3f PositionDelta = RestControlPoint - InterpolationData.SampleRestPositions[SampleIndex];
-					const float FunctionValue = FMath::Sqrt(FVector3f::DotProduct(PositionDelta, PositionDelta) + 1);
-					ControlPoint += FunctionValue * SampleDeformations[SampleIndex];
-				}
-				ControlPoint += SampleDeformations[SampleCount];
-				ControlPoint += SampleDeformations[SampleCount + 1] * RestControlPoint.X;
-				ControlPoint += SampleDeformations[SampleCount + 2] * RestControlPoint.Y;
-				ControlPoint += SampleDeformations[SampleCount + 3] * RestControlPoint.Z;
-
-				return ControlPoint;
+				return DeformPoint(InterpolationData, SampleDeformations, RestControlPoint);
 			};
 
 			const int32 NumVertexIndices = DeformingMesh.VertexCount();
@@ -118,23 +126,23 @@ namespace UE::MeshResizing
 				int32 EndIndex = (BatchIndex + 1) * VerticesByTask;
 				EndIndex = BatchIndex == NumBatches - 1 ? FMath::Min(NumVertexIndices, EndIndex) : EndIndex;
 
-				UE::Tasks::FTask PendingTask = UE::Tasks::Launch(UE_SOURCE_LOCATION, [&DeformingMesh, StartIndex, EndIndex, BatchIndex, DeformPoint, &NormalOverlays]
+				UE::Tasks::FTask PendingTask = UE::Tasks::Launch(UE_SOURCE_LOCATION, [&DeformingMesh, StartIndex, EndIndex, BatchIndex, DeformPointLambda, &NormalOverlays]
 				{
 					for (int32 VertexIndex = StartIndex; VertexIndex < EndIndex; VertexIndex++)
 					{
 						check(DeformingMesh.GetVerticesRefCounts().IsValid(VertexIndex));
 						const FVector3f RestPoint(FVector3f(DeformingMesh.GetVertexRef(VertexIndex)));
-						const FVector3f DeformedPoint = DeformPoint(RestPoint);
+						const FVector3f DeformedPoint = DeformPointLambda(RestPoint);
 
 						DeformingMesh.SetVertex(VertexIndex, FVector3d(DeformedPoint));
 
 						for (UE::Geometry::FDynamicMeshNormalOverlay* const Overlay : NormalOverlays)
 						{
 							Overlay->EnumerateVertexElements(VertexIndex,
-								[&RestPoint, &DeformedPoint, &DeformPoint, Overlay](int32 TriangleID, int32 ElementID, const FVector3f& Normal)
+								[&RestPoint, &DeformedPoint, &DeformPointLambda, Overlay](int32 TriangleID, int32 ElementID, const FVector3f& Normal)
 								{
 									const FVector3f RestNormalEnd = RestPoint + Normal;
-									const FVector3f DeformedNormalEnd = DeformPoint(RestNormalEnd);
+									const FVector3f DeformedNormalEnd = DeformPointLambda(RestNormalEnd);
 									const FVector3f DeformedNorml = (DeformedNormalEnd - DeformedPoint).GetSafeNormal();
 									Overlay->SetElement(ElementID, DeformedNorml);
 									return true;
@@ -147,6 +155,53 @@ namespace UE::MeshResizing
 			}
 			UE::Tasks::Wait(PendingTasks);
 		}
+
+		// From GroomRBFDeformer
+		static void ApplyRBFDeformation(const FMeshResizingRBFInterpolationData& InterpolationData, const TArray<FVector3f>& SampleDeformations, TArrayView<FVector3f>& Positions, TArray<TArrayView<FVector3f>> NormalsArray)
+		{
+			const int32 SampleCount = InterpolationData.SampleRestPositions.Num();
+			check(SampleCount + 4 == SampleDeformations.Num());
+
+			auto DeformPointLambda = [&SampleDeformations, &InterpolationData](const FVector3f& RestControlPoint)
+				{
+					return DeformPoint(InterpolationData, SampleDeformations, RestControlPoint);
+				};
+
+			const int32 NumVertexIndices = Positions.Num();
+			const int32 NumTasks = FMath::Max(FTaskGraphInterface::Get().GetNumWorkerThreads(), NumVertexIndices);
+			constexpr int32 MinVerticesByTask = 10;
+			const int32 VerticesByTask = FMath::Max(FMath::Min(FMath::DivideAndRoundUp(NumVertexIndices, NumTasks), MinVerticesByTask), 1);
+			const int32 NumBatches = FMath::DivideAndRoundUp(NumVertexIndices, VerticesByTask);
+			TArray<UE::Tasks::FTask> PendingTasks;
+			PendingTasks.Reserve(NumBatches);
+
+			for (int32 BatchIndex = 0; BatchIndex < NumBatches; BatchIndex++)
+			{
+				const int32 StartIndex = BatchIndex * VerticesByTask;
+				int32 EndIndex = (BatchIndex + 1) * VerticesByTask;
+				EndIndex = BatchIndex == NumBatches - 1 ? FMath::Min(NumVertexIndices, EndIndex) : EndIndex;
+
+				UE::Tasks::FTask PendingTask = UE::Tasks::Launch(UE_SOURCE_LOCATION, [&Positions, StartIndex, EndIndex, BatchIndex, DeformPointLambda, &NormalsArray]
+					{
+						for (int32 VertexIndex = StartIndex; VertexIndex < EndIndex; VertexIndex++)
+						{
+							const FVector3f RestPoint(Positions[VertexIndex]);
+							const FVector3f DeformedPoint = DeformPointLambda(RestPoint);
+
+							Positions[VertexIndex] = DeformedPoint;
+
+							for (TArrayView<FVector3f>& Normals : NormalsArray)
+							{								
+								const FVector3f RestNormalEnd = RestPoint + Normals[VertexIndex];
+								const FVector3f DeformedNormalEnd = DeformPointLambda(RestNormalEnd);
+								Normals[VertexIndex] = (DeformedNormalEnd - DeformedPoint).GetSafeNormal();
+							}
+						}
+					});
+				PendingTasks.Add(PendingTask);
+			}
+			UE::Tasks::Wait(PendingTasks);
+		}
 	}
 
 	void FRBFInterpolation::GenerateWeights(const UE::Geometry::FDynamicMesh3& BaseMesh, int32 NumInterpolationPoints, FMeshResizingRBFInterpolationData& OutData)
@@ -195,11 +250,16 @@ namespace UE::MeshResizing
 		ValidPoints.Init(true, SourcePositions.Num());
 
 		GroomBinding_RBFWeighting::FPointsSampler PointsSampler(ValidPoints, SourcePositions.GetData(), NumInterpolationPoints);
-		const uint32 SampleCount = PointsSampler.SamplePositions.Num();
+		const int32 SampleCount = PointsSampler.SamplePositions.Num();
 		GroomBinding_RBFWeighting::FWeightsBuilder InterpolationWeights(SampleCount, SampleCount, PointsSampler.SamplePositions.GetData(), PointsSampler.SamplePositions.GetData());
 		Private::UpdateInterpolationData(PointsSampler, InterpolationWeights, OutData);
 	}
 
+	void FRBFInterpolation::GenerateMeshSamples(const TArray<FVector3f>& TargetPositions, const FMeshResizingRBFInterpolationData& InterpolationData, TArray<FVector3f>& SampleDeformations)
+	{
+		Private::UpdateMeshSamples(TargetPositions, InterpolationData, SampleDeformations);
+	}
+
 	void FRBFInterpolation::DeformPoints(const TArray<FVector3f>& TargetPositions, const FMeshResizingRBFInterpolationData& InterpolationData, bool bInterpolateNormals, UE::Geometry::FDynamicMesh3& DeformingMesh)
 	{
 		TArray<FVector3f> SampleDeformations;
@@ -207,6 +267,68 @@ namespace UE::MeshResizing
 		Private::ApplyRBFDeformation(InterpolationData, SampleDeformations, bInterpolateNormals, DeformingMesh);
 	}
 
+	void FRBFInterpolation::DeformCoordinateFrames(const TArray<FVector3f>& TargetPositions, const FMeshResizingRBFInterpolationData& InterpolationData, bool bNormalize, bool bOrthogonalize, TArray<FMatrix44f>& Coordinates)
+	{
+		TArray<FVector3f> SampleDeformations;
+		Private::UpdateMeshSamples(TargetPositions, InterpolationData, SampleDeformations);
+
+		for (FMatrix44f& Coord : Coordinates)
+		{
+			FVector3f Origin = Coord.GetOrigin();
+			FVector3f TangentU, TangentV, Normal;
+			Coord.GetScaledAxes(TangentU, TangentV, Normal);
+
+			FVector3f TangentUEnd = Origin + TangentU;
+			FVector3f TangentVEnd = Origin + TangentV;
+			FVector3f NormalEnd = Origin + Normal;
+			Origin = Private::DeformPoint(InterpolationData, SampleDeformations, Origin);
+			TangentUEnd = Private::DeformPoint(InterpolationData, SampleDeformations, TangentUEnd);
+			TangentVEnd = Private::DeformPoint(InterpolationData, SampleDeformations, TangentVEnd);
+			NormalEnd = Private::DeformPoint(InterpolationData, SampleDeformations, NormalEnd);
+
+			TangentU = TangentUEnd - Origin;
+			TangentV = TangentVEnd - Origin;
+			Normal = NormalEnd - Origin;
+
+			if (bOrthogonalize)
+			{
+				if (bNormalize)
+				{
+					Normal = Normal.GetSafeNormal();
+					TangentU = (TangentU - (TangentU | Normal) * Normal).GetSafeNormal();
+					TangentV = Normal.Cross(TangentU);
+				}
+				else
+				{
+					const float NormalLenSq = Normal.SquaredLength();
+					if (NormalLenSq > UE_SMALL_NUMBER)
+					{
+						TangentU = (TangentU - (TangentU | Normal) * Normal / NormalLenSq);
+					}
+					TangentV = (TangentV.Length() / (FMath::Sqrt(NormalLenSq) * TangentU.Length())) * Normal.Cross(TangentU);
+				}
+			}
+			else if (bNormalize)
+			{
+				Normal = Normal.GetSafeNormal();
+				TangentU = TangentU.GetSafeNormal();
+				TangentV = TangentV.GetSafeNormal();
+			}
+
+			Coord.SetAxes(&TangentU, &TangentV, &Normal, &Origin);
+		}
+	}
+
+	void FRBFInterpolation::DeformPoints(const TArray<FVector3f>& TargetPositions, const FMeshResizingRBFInterpolationData& InterpolationData, TArray<FVector3d>& Points)
+	{
+		TArray<FVector3f> SampleDeformations;
+		Private::UpdateMeshSamples(TargetPositions, InterpolationData, SampleDeformations);
+		for (FVector3d& Point : Points)
+		{
+			Point = FVector3d(Private::DeformPoint(InterpolationData, SampleDeformations, FVector3f(Point)));
+		}
+	}
+
 	void FRBFInterpolation::DeformPoints(const FMeshDescription& TargetMesh, const FMeshResizingRBFInterpolationData& InterpolationData, bool bInterpolateNormals, UE::Geometry::FDynamicMesh3& DeformingMesh)
 	{
 		TArray<FVector3f> TargetPositions;
@@ -227,4 +349,50 @@ namespace UE::MeshResizing
 
 		DeformPoints(TargetPositions, InterpolationData, bInterpolateNormals, DeformingMesh);
 	}
+
+	void FRBFInterpolation::DeformPoints(const FMeshDescription& TargetMesh, const FMeshResizingRBFInterpolationData& InterpolationData, TArrayView<FVector3f> Points, TArrayView<FVector3f> Normals)
+	{
+		TArray<FVector3f> TargetPositions;
+		TargetPositions.Reserve(InterpolationData.SampleIndices.Num());
+		for (int32 Index = 0; Index < InterpolationData.SampleIndices.Num(); ++Index)
+		{
+			const int32 SampleIndex = InterpolationData.SampleIndices[Index];
+			if (ensure(TargetMesh.IsVertexValid(SampleIndex)))
+			{
+				TargetPositions.Emplace(TargetMesh.GetVertexPosition(SampleIndex));
+			}
+			else
+			{
+				// Just use Source positions
+				TargetPositions.Emplace(InterpolationData.SampleRestPositions[Index]);
+			}
+		}
+
+		TArray<FVector3f> SampleDeformations;
+		Private::UpdateMeshSamples(TargetPositions, InterpolationData, SampleDeformations); 
+		Private::ApplyRBFDeformation(InterpolationData, SampleDeformations, Points, { Normals });
+	}
+
+	void FRBFInterpolation::DeformPoints(const FMeshDescription& TargetMesh, const FMeshResizingRBFInterpolationData& InterpolationData, TArrayView<FVector3f> Points, TArrayView<FVector3f> Normals, TArrayView<FVector3f> TangentUs, TArrayView<FVector3f> TangentVs)
+	{
+		TArray<FVector3f> TargetPositions;
+		TargetPositions.Reserve(InterpolationData.SampleIndices.Num());
+		for (int32 Index = 0; Index < InterpolationData.SampleIndices.Num(); ++Index)
+		{
+			const int32 SampleIndex = InterpolationData.SampleIndices[Index];
+			if (ensure(TargetMesh.IsVertexValid(SampleIndex)))
+			{
+				TargetPositions.Emplace(TargetMesh.GetVertexPosition(SampleIndex));
+			}
+			else
+			{
+				// Just use Source positions
+				TargetPositions.Emplace(InterpolationData.SampleRestPositions[Index]);
+			}
+		}
+
+		TArray<FVector3f> SampleDeformations;
+		Private::UpdateMeshSamples(TargetPositions, InterpolationData, SampleDeformations);
+		Private::ApplyRBFDeformation(InterpolationData, SampleDeformations, Points, { Normals, TangentUs, TangentVs });
+	}
 }
\ No newline at end of file
diff --git a/Engine/Plugins/Experimental/MeshResizing/Source/MeshResizingCore/Public/MeshResizing/RBFInterpolation.h b/Engine/Plugins/Experimental/MeshResizing/Source/MeshResizingCore/Public/MeshResizing/RBFInterpolation.h
--- a/Engine/Plugins/Experimental/MeshResizing/Source/MeshResizingCore/Public/MeshResizing/RBFInterpolation.h
+++ b/Engine/Plugins/Experimental/MeshResizing/Source/MeshResizingCore/Public/MeshResizing/RBFInterpolation.h
@@ -30,9 +30,14 @@ namespace UE::MeshResizing
 	{
 		static void GenerateWeights(const TConstArrayView<FVector3f>& SourcePositions, int32 NumInterpolationPoints, FMeshResizingRBFInterpolationData& OutData);
 		static void GenerateWeights(const UE::Geometry::FDynamicMesh3& BaseMesh, int32 NumInterpolationPoints, FMeshResizingRBFInterpolationData& OutData);
+		static void GenerateMeshSamples(const TArray<FVector3f>& TargetPositions, const FMeshResizingRBFInterpolationData& InterpolationData, TArray<FVector3f>& SampleDeformations);
 		static void GenerateWeights(const FMeshDescription& BaseMesh, int32 NumInterpolationPoints, FMeshResizingRBFInterpolationData& OutData);
 		static void DeformPoints(const TArray<FVector3f>& TargetPositions, const FMeshResizingRBFInterpolationData& InterpolationData, bool bInterpolateNormals,  UE::Geometry::FDynamicMesh3& DeformingMesh);
+		static void DeformCoordinateFrames(const TArray<FVector3f>& TargetPositions, const FMeshResizingRBFInterpolationData& InterpolationData, bool bNormalize, bool bOrthogonalize, TArray<FMatrix44f>& Coordinates);
+		static void DeformPoints(const TArray<FVector3f>& TargetPositions, const FMeshResizingRBFInterpolationData& InterpolationData, TArray<FVector3d>& Points);
 		static void DeformPoints(const FMeshDescription& TargetMesh, const FMeshResizingRBFInterpolationData& InterpolationData, bool bInterpolateNormals, UE::Geometry::FDynamicMesh3& DeformingMesh);
+		static void DeformPoints(const FMeshDescription& TargetMesh, const FMeshResizingRBFInterpolationData& InterpolationData, TArrayView<FVector3f> Points, TArrayView<FVector3f> Normals);
+		static void DeformPoints(const FMeshDescription& TargetMesh, const FMeshResizingRBFInterpolationData& InterpolationData, TArrayView<FVector3f> Points, TArrayView<FVector3f> Normals, TArrayView<FVector3f> TangentUs, TArrayView<FVector3f> TangentVs);
 	};
 }
 
