Chaos Outfit Asset - Added UV texture resizing controller.

#rnx

diff --git a/Engine/Plugins/ChaosClothAssetEditor/Source/ChaosClothAssetDataflowNodes/Private/ChaosClothAsset/ClothCollectionToDynamicMeshNode.cpp b/Engine/Plugins/ChaosClothAssetEditor/Source/ChaosClothAssetDataflowNodes/Private/ChaosClothAsset/ClothCollectionToDynamicMeshNode.cpp
--- a/Engine/Plugins/ChaosClothAssetEditor/Source/ChaosClothAssetDataflowNodes/Private/ChaosClothAsset/ClothCollectionToDynamicMeshNode.cpp
+++ b/Engine/Plugins/ChaosClothAssetEditor/Source/ChaosClothAssetDataflowNodes/Private/ChaosClothAsset/ClothCollectionToDynamicMeshNode.cpp
@@ -82,6 +82,7 @@ FChaosClothAssetUpdateClothFromDynamicMeshNode::FChaosClothAssetUpdateClothFromD
 {
 	RegisterInputConnection(&Collection);
 	RegisterInputConnection(&DynamicMesh);
+	RegisterInputConnection(&Materials);
 	RegisterOutputConnection(&Collection, &Collection);
 }
 
@@ -96,6 +97,15 @@ void FChaosClothAssetUpdateClothFromDynamicMeshNode::Evaluate(UE::Dataflow::FCon
 		FCollectionClothFacade ClothFacade(ClothCollection);
 		if (ClothFacade.IsValid())
 		{
+			if (bCopyToRenderMaterials)
+			{
+				const TArray<TObjectPtr<const UMaterialInterface>> InMaterials = GetValue(Context, &Materials);
+				TArrayView<FString> RenderMaterialPathNames = ClothFacade.GetRenderMaterialPathName();
+				for (int32 MaterialIndex = 0; MaterialIndex < FMath::Min(RenderMaterialPathNames.Num(), InMaterials.Num()); ++MaterialIndex)
+				{
+					RenderMaterialPathNames[MaterialIndex] = InMaterials[MaterialIndex]->GetPathName();
+				}
+			}
 			if (const TObjectPtr<const UDynamicMesh> InMesh = GetValue(Context, &DynamicMesh))
 			{
 				const UE::Geometry::FDynamicMesh3& DynMesh = InMesh->GetMeshRef();
@@ -163,30 +173,40 @@ void FChaosClothAssetUpdateClothFromDynamicMeshNode::Evaluate(UE::Dataflow::FCon
 				{
 					if (const UE::Geometry::FDynamicMeshAttributeSet* const AttributeSet = DynMesh.Attributes())
 					{
-						if (const UE::Geometry::FDynamicMeshUVOverlay* const UVOverlay = AttributeSet->GetUVLayer(UVChannelIndex))
+						const int32 UVStart = (UVChannelIndex == -1) ? 0: UVChannelIndex;
+						const int32 UVCount = (UVChannelIndex == -1) ? AttributeSet->NumUVLayers() : 1;
+
+						for (int32 UVIndex = UVStart; UVIndex < UVStart + UVCount; ++UVIndex)
 						{
-							TArrayView<TArray<FVector2f>> RenderUVs = ClothFacade.GetRenderUVs();
-							TConstArrayView<FIntVector> RenderTriangles = ClothFacade.GetRenderIndices();
-							for (int32 VertexID : DynMesh.VertexIndicesItr())
+							if (const UE::Geometry::FDynamicMeshUVOverlay* const UVOverlay = AttributeSet->GetUVLayer(UVIndex))
 							{
-								const int32 RenderMeshVertexID = ClothMapping.GetOriginalVertexID(VertexID);
-								if (RenderUVs.IsValidIndex(RenderMeshVertexID) && RenderUVs[RenderMeshVertexID].Num() > UVChannelIndex)
+								TArrayView<TArray<FVector2f>> RenderUVs = ClothFacade.GetRenderUVs();
+								TConstArrayView<FIntVector> RenderTriangles = ClothFacade.GetRenderIndices();
+								for (int32 VertexID : DynMesh.VertexIndicesItr())
 								{
-									// Get the UV values corresponding to VertexID in the dynamic mesh
-									UVOverlay->EnumerateVertexElements(VertexID,
-										[&RenderUVs, &ClothMapping, &RenderTriangles, RenderMeshVertexID, UVChannelIndex=this->UVChannelIndex](int32 TriangleID, int32 ElementID, const FVector2f& UVValue)->bool
+									const int32 RenderMeshVertexID = ClothMapping.GetOriginalVertexID(VertexID);
+									if (RenderUVs.IsValidIndex(RenderMeshVertexID))
+									{
+										if (UVIndex >= RenderUVs[RenderMeshVertexID].Num())
 										{
-											const int32 RenderMeshTriID = ClothMapping.GetOriginalTriangleID(TriangleID);
-											const FIntVector& RenderTri = RenderTriangles[RenderMeshTriID];
-											for (int32 LocalIndex = 0; LocalIndex < 3; ++LocalIndex)
+											RenderUVs[RenderMeshVertexID].AddZeroed();
+										}
+										// Get the UV values corresponding to VertexID in the dynamic mesh
+										UVOverlay->EnumerateVertexElements(VertexID,
+											[&RenderUVs, &ClothMapping, &RenderTriangles, RenderMeshVertexID, UVIndex](int32 TriangleID, int32 ElementID, const FVector2f& UVValue)->bool
 											{
-												if (RenderTri[LocalIndex] == RenderMeshVertexID)
-												{ 
-													RenderUVs[RenderTri[LocalIndex]][UVChannelIndex] = UVValue;
+												const int32 RenderMeshTriID = ClothMapping.GetOriginalTriangleID(TriangleID);
+												const FIntVector& RenderTri = RenderTriangles[RenderMeshTriID];
+												for (int32 LocalIndex = 0; LocalIndex < 3; ++LocalIndex)
+												{
+													if (RenderTri[LocalIndex] == RenderMeshVertexID)
+													{ 
+														RenderUVs[RenderTri[LocalIndex]][UVIndex] = UVValue;
+													}
 												}
-											}
-											return true;
-										});
+												return true;
+											});
+									}
 								}
 							}
 						}
diff --git a/Engine/Plugins/ChaosClothAssetEditor/Source/ChaosClothAssetDataflowNodes/Public/ChaosClothAsset/ClothCollectionToDynamicMeshNode.h b/Engine/Plugins/ChaosClothAssetEditor/Source/ChaosClothAssetDataflowNodes/Public/ChaosClothAsset/ClothCollectionToDynamicMeshNode.h
--- a/Engine/Plugins/ChaosClothAssetEditor/Source/ChaosClothAssetDataflowNodes/Public/ChaosClothAsset/ClothCollectionToDynamicMeshNode.h
+++ b/Engine/Plugins/ChaosClothAssetEditor/Source/ChaosClothAssetDataflowNodes/Public/ChaosClothAsset/ClothCollectionToDynamicMeshNode.h
@@ -71,7 +71,10 @@ private:
 	FManagedArrayCollection Collection;
 
 	UPROPERTY(Meta = (DataflowInput))
-	TObjectPtr<UDynamicMesh> DynamicMesh;
+	TObjectPtr<const UDynamicMesh> DynamicMesh;
+
+	UPROPERTY(Meta = (DataflowInput))
+	TArray<TObjectPtr<const UMaterialInterface>> Materials;
 
 	/** Copy DynamicMesh Vertex Positions to Render Positions */
 	UPROPERTY(EditAnywhere, Category = "DynamicMesh To Cloth Collection")
@@ -85,6 +88,10 @@ private:
 	UPROPERTY(EditAnywhere, Category = "DynamicMesh To Cloth Collection")
 	bool bCopyUVsToRenderUVs = false;
 
+	/** Copy input materials to Render Materials (order and number must match otherwise only the minimum common number of materials are updated) */
+	UPROPERTY(EditAnywhere, Category = "DynamicMesh To Cloth Collection")
+	bool bCopyToRenderMaterials = false;
+
 	/** Copy DynamicMesh Vertex Positions to Sim3D Positions */
 	UPROPERTY(EditAnywhere, Category = "DynamicMesh To Cloth Collection")
 	bool bCopyToSim3DPositions = false;
@@ -97,7 +104,7 @@ private:
 	UPROPERTY(EditAnywhere, Category = "DynamicMesh To Cloth Collection")
 	bool bCopyUVsToSim2DPositions = false;
 
-	/** Which UV Channel to use at Sim2D Positions or Render UVs */
+	/** Which UV Channel to use at Sim2D Positions or Render UVs. Use -1 to copy all Render UVs. */
 	UPROPERTY(EditAnywhere, Category = "DynamicMesh To Cloth Collection", meta = (EditCondition = "bCopyUVsToSim2DPositions || bCopyUVsToRenderUVs", EditConditionHides))
 	int32 UVChannelIndex = 0;
 };
diff --git a/Engine/Plugins/Experimental/MeshResizing/Source/MeshResizingNodes/Private/MeshResizing/AlignUVMeshNode.cpp b/Engine/Plugins/Experimental/MeshResizing/Source/MeshResizingNodes/Private/MeshResizing/AlignUVMeshNode.cpp
--- a/Engine/Plugins/Experimental/MeshResizing/Source/MeshResizingNodes/Private/MeshResizing/AlignUVMeshNode.cpp
+++ b/Engine/Plugins/Experimental/MeshResizing/Source/MeshResizingNodes/Private/MeshResizing/AlignUVMeshNode.cpp
@@ -81,6 +81,7 @@ FAlignUVMeshNode::FAlignUVMeshNode(const UE::Dataflow::FNodeParameters& InParam,
 	RegisterInputConnection(&ResizingMesh);
 	RegisterInputConnection(&BaseMesh);
 	RegisterInputConnection(&UVChannelIndex);
+	RegisterInputConnection(&BaseUVChannelIndex);
 	RegisterOutputConnection(&ResizingMesh, &ResizingMesh);
 	RegisterOutputConnection(&UVChannelIndex, &UVChannelIndex);
 }
@@ -106,8 +107,13 @@ void FAlignUVMeshNode::Evaluate(UE::Dataflow::FContext& Context, const FDataflow
 					if (InResizingMesh->GetDynamicMeshRef().HasAttributes() && InBaseMesh->GetDynamicMeshRef().HasAttributes())
 					{
 						const int32 UVChannel = GetValue(Context, &UVChannelIndex);
+						int32 BaseUVChannel = GetValue(Context, &BaseUVChannelIndex);
+						if (BaseUVChannel == -1)
+						{
+							BaseUVChannel = UVChannel;
+						}
 
-						if (UVChannel >= 0 && InResizingMesh->GetDynamicMeshRef().Attributes()->NumUVLayers() > UVChannel && InBaseMesh->GetDynamicMeshRef().Attributes()->NumUVLayers() > UVChannel)
+						if (UVChannel >= 0 && InResizingMesh->GetDynamicMeshRef().Attributes()->NumUVLayers() > UVChannel && InBaseMesh->GetDynamicMeshRef().Attributes()->NumUVLayers() > BaseUVChannel)
 						{
 							const FDynamicMeshUVOverlay* const InResizedUVOverlay = InResizingMesh->GetDynamicMeshRef().Attributes()->GetUVLayer(UVChannel);
 
@@ -121,7 +127,7 @@ void FAlignUVMeshNode::Evaluate(UE::Dataflow::FContext& Context, const FDataflow
 								return ResizedUVOverlay ? ResizedUVOverlay->AreTrianglesConnected(Triangle0, Triangle1) : false;
 							});
 
-							const FDynamicMeshUVOverlay* const BaseUVOverlay = InBaseMesh->GetDynamicMeshRef().Attributes()->GetUVLayer(UVChannel);
+							const FDynamicMeshUVOverlay* const BaseUVOverlay = InBaseMesh->GetDynamicMeshRef().Attributes()->GetUVLayer(BaseUVChannel);
 
 							for (const FMeshConnectedComponents::FComponent& Island : UVIslands.Components)
 							{
@@ -176,7 +182,7 @@ void FAlignUVMeshNode::Evaluate(UE::Dataflow::FContext& Context, const FDataflow
 						}
 						else
 						{
-							Context.Warning(TEXT("Invalid UVChannelIndex"), this, Out);
+							Context.Warning(TEXT("Invalid UVChannelIndex or BaseUVChannelIndex"), this, Out);
 						}
 					}
 					else
diff --git a/Engine/Plugins/Experimental/MeshResizing/Source/MeshResizingNodes/Private/MeshResizing/AlignUVMeshNode.h b/Engine/Plugins/Experimental/MeshResizing/Source/MeshResizingNodes/Private/MeshResizing/AlignUVMeshNode.h
--- a/Engine/Plugins/Experimental/MeshResizing/Source/MeshResizingNodes/Private/MeshResizing/AlignUVMeshNode.h
+++ b/Engine/Plugins/Experimental/MeshResizing/Source/MeshResizingNodes/Private/MeshResizing/AlignUVMeshNode.h
@@ -29,6 +29,10 @@ private:
 	UPROPERTY(EditAnywhere, Category = "Align", meta = (DataflowInput, DataflowOutput, DataflowPassthrough = "UVChannelIndex", ClampMin = 0, ClampMax = 7, UIMin = 0, UIMax = 7))
 	int32 UVChannelIndex = 0;
 
+	/** Base UV channel index in case it differs from the ResizingMesh UV channel index, or -1 to use the same channel. */
+	UPROPERTY(EditAnywhere, Category = "Align", meta = (DataflowInput, ClampMin = -1, ClampMax = 7, UIMin = -1, UIMax = 7))
+	int32 BaseUVChannelIndex = -1;
+
 	UPROPERTY(EditAnywhere, Category = "Align")
 	bool bScale = true;
 
diff --git a/Engine/Plugins/Experimental/MeshResizing/Source/MeshResizingNodes/Private/MeshResizing/MeshResizingDataflowNodesModule.cpp b/Engine/Plugins/Experimental/MeshResizing/Source/MeshResizingNodes/Private/MeshResizing/MeshResizingDataflowNodesModule.cpp
--- a/Engine/Plugins/Experimental/MeshResizing/Source/MeshResizingNodes/Private/MeshResizing/MeshResizingDataflowNodesModule.cpp
+++ b/Engine/Plugins/Experimental/MeshResizing/Source/MeshResizingNodes/Private/MeshResizing/MeshResizingDataflowNodesModule.cpp
@@ -11,6 +11,7 @@
 #include "MeshResizing/UVMeshTransformNode.h"
 #include "MeshResizing/AlignUVMeshNode.h"
 #include "MeshResizing/MeshResizingTextureNodes.h"
+#include "MeshResizing/UVResizeControllerNode.h"
 
 namespace UE::MeshResizing
 {
@@ -37,6 +38,7 @@ namespace UE::MeshResizing
 			DATAFLOW_NODE_REGISTER_CREATION_FACTORY(FMeshWarpNode);
 			DATAFLOW_NODE_REGISTER_CREATION_FACTORY(FGenerateRBFResizingWeightsNode);
 			DATAFLOW_NODE_REGISTER_CREATION_FACTORY(FApplyRBFResizingNode);
+			DATAFLOW_NODE_REGISTER_CREATION_FACTORY(FUVResizeControllerNode);
 		}
 
 	}  // End namespace Private
diff --git a/Engine/Plugins/Experimental/MeshResizing/Source/MeshResizingNodes/Private/MeshResizing/UVResizeControllerNode.cpp b/Engine/Plugins/Experimental/MeshResizing/Source/MeshResizingNodes/Private/MeshResizing/UVResizeControllerNode.cpp
new file mode 100644
--- /dev/null
+++ b/Engine/Plugins/Experimental/MeshResizing/Source/MeshResizingNodes/Private/MeshResizing/UVResizeControllerNode.cpp
@@ -0,0 +1,228 @@
+// Copyright Epic Games, Inc. All Rights Reserved.
+#include "MeshResizing/UVResizeControllerNode.h"
+#include "AssetRegistry/AssetRegistryModule.h"
+#include "Dataflow/DataflowMesh.h"
+#include "DynamicMesh/DynamicMeshAttributeSet.h"
+#include "Engine/Texture.h"
+#include "Materials/MaterialInstanceConstant.h"
+#include "Misc/PackageName.h"
+
+#include UE_INLINE_GENERATED_CPP_BY_NAME(UVResizeControllerNode)
+
+namespace UE::Chaos::OutfitAsset::Private
+{
+	static constexpr int32 MaxUVChannelsForResizing = MAX_TEXCOORDS / 2;
+
+	static int32 GetNumUVLayers(const TObjectPtr<const UDataflowMesh> DataflowMesh)
+	{
+		using namespace UE::Geometry;
+
+		if (const FDynamicMesh3* const DynamicMesh = DataflowMesh->GetDynamicMesh())
+		{
+			if (const FDynamicMeshAttributeSet* const Attributes = DynamicMesh->Attributes())
+			{
+				return Attributes->NumUVLayers();
+			}
+		}
+		return 0;
+	}
+
+	// Find all UV channel indices parameter names
+	// A good candidate must:
+	// 1. Have a matching name to a texture parameter name, which only differ by their predetermined suffixes.
+	// 2. Have the UVChannel index value be less than the number of UV layers.
+	// 3. Be related to a texture that is set to wrap around.
+	static TArray<FString> FindUVChannelParameterNames(
+		const TObjectPtr<const UMaterialInterface>& Material,
+		const FString& TextureSuffix,
+		const FString& UVChannelSuffix,
+		const int32 NumUVLayers)
+	{
+		TArray<FString> UVChannelParameterNames;
+
+		if (Material)
+		{
+			// Iterate through all textures linked to this material
+			TArray<FMaterialParameterInfo> TextureParameterInfos;
+			TArray<FGuid> TextureParameterIds;
+			Material->GetAllTextureParameterInfo(TextureParameterInfos, TextureParameterIds);
+
+			for (const FMaterialParameterInfo& TextureParameterInfo : TextureParameterInfos)
+			{
+				// Only deals with textures with name ending by TextureSuffix and replace it with UVChannelSuffix to find the matching UV Channel index
+				const FString TextureParameterName = TextureParameterInfo.Name.ToString();
+				if (TextureParameterName.EndsWith(TextureSuffix))
+				{
+					// Found a texture with suffix, try to locate the texture object
+					UTexture* Texture;
+					if (Material->GetTextureParameterValue(TextureParameterInfo, Texture) && Texture)
+					{
+						if (Texture->GetTextureAddressX() == TA_Wrap &&
+							Texture->GetTextureAddressY() == TA_Wrap)
+						{
+							// Texture can wrap around, so check for a matching channel parameter
+							const FString UVChannelParameterName = TextureParameterName.LeftChop(TextureSuffix.Len()) + UVChannelSuffix;
+							float Value;
+							if (Material->GetScalarParameterValue(*UVChannelParameterName, Value))
+							{
+								const int32 UVChannelIndex = (int32)Value;
+								if (UVChannelIndex < NumUVLayers)
+								{
+									UVChannelParameterNames.Emplace(UVChannelParameterName);
+								}
+							}
+						}
+					}
+				}
+			}
+		}
+		return UVChannelParameterNames;
+	}
+}
+
+FUVResizeControllerNode::FUVResizeControllerNode(const UE::Dataflow::FNodeParameters& InParam, FGuid InGuid)
+	: FDataflowNode(InParam, InGuid)
+{
+	RegisterInputConnection(&Mesh);
+	RegisterOutputConnection(&Mesh, &Mesh);
+	RegisterOutputConnection(&UVChannelIndices);
+	RegisterOutputConnection(&SourceUVChannelIndices);
+	RegisterOutputConnection(&bHasUVChannelsToResize);
+}
+
+void FUVResizeControllerNode::Evaluate(UE::Dataflow::FContext& Context, const FDataflowOutput* Out) const
+{
+	using namespace UE::Chaos::OutfitAsset;
+
+	if (Out->IsA(&Mesh) || Out->IsA(&UVChannelIndices) || Out->IsA(&bHasUVChannelsToResize))
+	{
+		bool bOutHasUVChannelsToResize = false;
+
+		// Check available layers
+		if (const TObjectPtr<UDataflowMesh> InMesh = GetValue(Context, &Mesh))
+		{
+#if WITH_EDITOR
+			if (GIsEditor)  // For UMaterialInstanceConstant::SetScalarParameterValueEditorOnly and UMaterialInstanceConstant::SetParentEditorOnly
+			{
+				TArray<int32> UVChannelRemaps;
+				UVChannelRemaps.Init(INDEX_NONE, MAX_TEXCOORDS);
+
+				const int32 NumUVLayers = Private::GetNumUVLayers(InMesh);
+				if (NumUVLayers > 0 && NumUVLayers < Private::MaxUVChannelsForResizing)
+				{
+					check(InMesh && InMesh->GetDynamicMesh() && InMesh->GetDynamicMesh()->Attributes());  // Already tested in GetNumUVLayers
+
+					// Create a new dynamic mesh
+					UE::Geometry::FDynamicMesh3 DynamicMesh;
+					DynamicMesh.Copy(*InMesh->GetDynamicMesh());
+					DynamicMesh.EnableAttributes();
+
+					const TArray<TObjectPtr<UMaterialInterface>>& InMaterials = InMesh->GetMaterials();
+					TArray<TObjectPtr<UMaterialInterface>> OutMaterials;
+					OutMaterials.Reserve(InMaterials.Num());
+
+					for (const TObjectPtr<UMaterialInterface>& Material : InMaterials)
+					{
+						const TArray<FString> UVChannelParameterNames = Private::FindUVChannelParameterNames(Material, TextureSuffix, UVChannelSuffix, NumUVLayers);
+						if (UVChannelParameterNames.Num())
+						{
+							check(Material);  // Already tested in FindUVChannelParameterNames
+
+							// Create a new material instance that has this Material as base
+							const FString MaterialPackagePath = FPackageName::GetLongPackagePath(Material->GetOutermost()->GetName());
+							FString MaterialName = Material->GetName() + TEXT("_UVResized");
+							FString MaterialPackageName = FPaths::Combine(MaterialPackagePath, MaterialName);
+							if (FindPackage(nullptr, *MaterialPackageName))
+							{
+								MaterialPackageName = MakeUniqueObjectName(nullptr, UPackage::StaticClass(), FName(MaterialPackageName)).ToString();
+								MaterialName = FPaths::GetBaseFilename(MaterialPackageName);
+							}
+							if (UPackage* const MaterialPackage = CreatePackage(*MaterialPackageName))
+							{
+								if (UMaterialInstanceConstant* const MaterialInstance = NewObject<UMaterialInstanceConstant>(MaterialPackage, FName(MaterialName), RF_Public | RF_Standalone))
+								{
+									MaterialInstance->PreEditChange(nullptr);
+									MaterialInstance->PostEditChange();  // Call to InitResources
+
+									MaterialInstance->SetParentEditorOnly(Material);
+
+									// Re-assign the UVChannels in this new material instance
+									for (const FString& UVChannelParameterName : UVChannelParameterNames)
+									{
+										float Value;
+										if (MaterialInstance->GetScalarParameterValue(*UVChannelParameterName, Value))
+										{
+											const int32 SourceUVChannelIndex = (int32)Value;
+											if (UVChannelRemaps.IsValidIndex(SourceUVChannelIndex))
+											{
+												int32& UVChannelIndex = UVChannelRemaps[SourceUVChannelIndex];
+												if (UVChannelIndex == INDEX_NONE)
+												{
+													// Duplicate UV channel
+													UVChannelIndex = DynamicMesh.Attributes()->NumUVLayers();
+													DynamicMesh.Attributes()->SetNumUVLayers(UVChannelIndex + 1);
+													const UE::Geometry::FDynamicMeshUVOverlay* const SourceUVLayer = DynamicMesh.Attributes()->GetUVLayer(SourceUVChannelIndex);
+													check(SourceUVLayer);
+													DynamicMesh.Attributes()->GetUVLayer(UVChannelIndex)->Copy(*SourceUVLayer);
+												}
+												MaterialInstance->SetScalarParameterValueEditorOnly(*UVChannelParameterName, (float)UVChannelIndex);
+											}
+										}
+									}
+
+									// Finalize the new material asset creation
+									MaterialInstance->PreEditChange(nullptr);
+									MaterialInstance->PostEditChange();
+									MaterialInstance->MarkPackageDirty();
+									FAssetRegistryModule::AssetCreated(MaterialInstance);
+
+									OutMaterials.Emplace(MaterialInstance);
+									continue;
+								}
+							}
+						}
+						OutMaterials.Emplace(Material);
+					}
+					// Create the output mesh
+					TObjectPtr<UDataflowMesh> OutMesh = NewObject<UDataflowMesh>();
+					OutMesh->SetDynamicMesh(MoveTemp(DynamicMesh));
+					OutMesh->SetMaterials(OutMaterials);
+					SetValue(Context, OutMesh, &Mesh);
+
+					TArray<int32> OutUVChannelIndices;
+					TArray<int32> OutSourceUVChannelIndices;
+					for (int32 Index = 0; Index < UVChannelRemaps.Num(); ++Index)
+					{
+						const int32 UVChannelRemap = UVChannelRemaps[Index];
+						if (UVChannelRemap != INDEX_NONE)
+						{
+							OutUVChannelIndices.Emplace(UVChannelRemap);
+							OutSourceUVChannelIndices.Emplace(Index);
+						}
+					}
+					SetValue(Context, OutUVChannelIndices.Num() > 0, &bHasUVChannelsToResize);  // Set before the MoveTemp!
+					SetValue(Context, MoveTemp(OutUVChannelIndices), &UVChannelIndices);
+					SetValue(Context, MoveTemp(OutSourceUVChannelIndices), &SourceUVChannelIndices);
+
+					return;
+				}
+				else
+				{
+					Context.Warning(
+						FString::Printf(TEXT("[%s] isn't a valid mesh with at least %d free UV channels to allow for UV resizing operations."),
+							*InMesh->GetName(),
+							Private::MaxUVChannelsForResizing),
+						this, Out);
+				}
+			}
+#endif  // #if WITH_EDITOR
+		}
+
+		TArray<int32> OutUVChannelIndices;
+		TArray<int32> OutSourceUVChannelIndices;
+		SafeForwardInput(Context, &Mesh, &Mesh);
+		SetValue(Context, MoveTemp(OutUVChannelIndices), &UVChannelIndices);
+		SetValue(Context, MoveTemp(OutSourceUVChannelIndices), &SourceUVChannelIndices);
+		SetValue(Context, false, &bHasUVChannelsToResize);
+	}
+}
diff --git a/Engine/Plugins/Experimental/MeshResizing/Source/MeshResizingNodes/Private/MeshResizing/UVResizeControllerNode.h b/Engine/Plugins/Experimental/MeshResizing/Source/MeshResizingNodes/Private/MeshResizing/UVResizeControllerNode.h
new file mode 100644
--- /dev/null
+++ b/Engine/Plugins/Experimental/MeshResizing/Source/MeshResizingNodes/Private/MeshResizing/UVResizeControllerNode.h
@@ -0,0 +1,53 @@
+// Copyright Epic Games, Inc. All Rights Reserved.
+
+#pragma once
+
+#include "Dataflow/DataflowNode.h"
+#include "UVResizeControllerNode.generated.h"
+
+class UDataflowMesh;
+class UMaterialInterface;
+
+/**
+ * UV Resizing logic.
+ * Returns whether this dynamic mesh is suitable for UV resizing and which UV channels to use.
+ */
+USTRUCT(Meta = (MeshResizing, Experimental))
+struct FUVResizeControllerNode : public FDataflowNode
+{
+	GENERATED_USTRUCT_BODY()
+	DATAFLOW_NODE_DEFINE_INTERNAL(FUVResizeControllerNode, "UVResizeController", "MeshResizing", "UV Resize Controller")
+	DATAFLOW_NODE_RENDER_TYPE("SurfaceRender", FName("UDataflowMesh"), "Mesh", "UVChannelIndex")
+
+public:
+	FUVResizeControllerNode(const UE::Dataflow::FNodeParameters& InParam, FGuid InGuid = FGuid::NewGuid());
+
+private:
+	/** The texture name suffix . */
+	UPROPERTY(EditAnywhere, Category = "UVResizeLogicNode")
+	FString TextureSuffix = TEXT("Texture");
+
+	/** The suffix to replace the texture name with pointing to the UV channel index scalar parameter. */
+	UPROPERTY(EditAnywhere, Category = "UVResizeLogicNode")
+	FString UVChannelSuffix = TEXT("UVIndex");
+
+	/** The input/output Dataflow mesh. */
+	UPROPERTY(Meta = (DataflowInput, DataflowOutput, DataflowPassthrough = "Mesh"))
+	TObjectPtr<UDataflowMesh> Mesh;
+
+	/** The UV channels to resize. */
+	UPROPERTY(Meta = (DataflowOutput))
+	TArray<int32> UVChannelIndices;
+
+	/** The matching UV channels on the original source mesh. */
+	UPROPERTY(Meta = (DataflowOutput))
+	TArray<int32> SourceUVChannelIndices;
+
+	/** Whether the input mesh has any UV channels to resize. */
+	UPROPERTY(Meta = (DataflowOutput))
+	bool bHasUVChannelsToResize = false;
+
+	//~ Begin FDataflowNode interface
+	virtual void Evaluate(UE::Dataflow::FContext& Context, const FDataflowOutput* Out) const override;
+	//~ End FDataflowNode interface
+};
