Fixup API macro usage

#rnx

diff --git a/Engine/Plugins/Animation/ControlRig/Source/ControlRig/Public/ControlRigOverride.h b/Engine/Plugins/Animation/ControlRig/Source/ControlRig/Public/ControlRigOverride.h
--- a/Engine/Plugins/Animation/ControlRig/Source/ControlRig/Public/ControlRigOverride.h
+++ b/Engine/Plugins/Animation/ControlRig/Source/ControlRig/Public/ControlRigOverride.h
@@ -25,7 +25,7 @@ public:
 
 	typedef  TFunction<void(const TCHAR*, ELogVerbosity::Type)> TReportFunction;
 
-	UE_API FControlRigOverrideValueErrorPipe( ELogVerbosity::Type InMaxVerbosity = ELogVerbosity::Warning, TReportFunction InReportFunction = nullptr )
+	FControlRigOverrideValueErrorPipe( ELogVerbosity::Type InMaxVerbosity = ELogVerbosity::Warning, TReportFunction InReportFunction = nullptr )
 		: FOutputDevice()
 		, NumErrors(0)
 		, MaxVerbosity(InMaxVerbosity)
@@ -33,7 +33,7 @@ public:
 	{
 	}
 
-	UE_API virtual void Serialize(const TCHAR* V, ELogVerbosity::Type Verbosity, const class FName& Category) override
+	virtual void Serialize(const TCHAR* V, ELogVerbosity::Type Verbosity, const class FName& Category) override
 	{
 		if(Verbosity <= MaxVerbosity)
 		{
@@ -45,8 +45,8 @@ public:
 		}
 	}
 
-	UE_API ELogVerbosity::Type GetMaxVerbosity() const { return MaxVerbosity; }
-	UE_API int32 GetNumErrors() const { return NumErrors; }
+	ELogVerbosity::Type GetMaxVerbosity() const { return MaxVerbosity; }
+	int32 GetNumErrors() const { return NumErrors; }
 
 private:
 	int32 NumErrors;
diff --git a/Engine/Plugins/Animation/IKRig/Source/IKRig/Public/Retargeter/RetargetOps/RetargetPoseOp.h b/Engine/Plugins/Animation/IKRig/Source/IKRig/Public/Retargeter/RetargetOps/RetargetPoseOp.h
--- a/Engine/Plugins/Animation/IKRig/Source/IKRig/Public/Retargeter/RetargetOps/RetargetPoseOp.h
+++ b/Engine/Plugins/Animation/IKRig/Source/IKRig/Public/Retargeter/RetargetOps/RetargetPoseOp.h
@@ -1,4 +1,4 @@
-﻿// Copyright Epic Games, Inc. All Rights Reserved.
+// Copyright Epic Games, Inc. All Rights Reserved.
 #pragma once
 
 #include "Retargeter/IKRetargetOps.h"
@@ -6,6 +6,8 @@
 
 #include "RetargetPoseOp.generated.h"
 
+#define UE_API IKRIG_API
+
 #define LOCTEXT_NAMESPACE "RetargetPoseOp"
 
 UENUM()
@@ -16,7 +18,7 @@ enum class ERetargetPoseMode : uint8
 };
 
 USTRUCT(BlueprintType, meta = (DisplayName = "Retarget Pose Op Settings"))
-struct IKRIG_API FIKRetargetPoseOpSettings : public FIKRetargetOpSettingsBase
+struct FIKRetargetPoseOpSettings : public FIKRetargetOpSettingsBase
 {
 	GENERATED_BODY()
 
@@ -29,44 +31,44 @@ struct IKRIG_API FIKRetargetPoseOpSettings : public FIKRetargetOpSettingsBase
 	UPROPERTY(EditAnywhere, BlueprintReadWrite, Category="Pose")
 	FName CurrentPose;
 	
-	virtual const UClass* GetControllerType() const override;
+	UE_API virtual const UClass* GetControllerType() const override;
 	
-	virtual void CopySettingsAtRuntime(const FIKRetargetOpSettingsBase* InSettingsToCopyFrom) override;
+	UE_API virtual void CopySettingsAtRuntime(const FIKRetargetOpSettingsBase* InSettingsToCopyFrom) override;
 };
 
 USTRUCT(BlueprintType, meta = (DisplayName = "Retarget Pose"))
-struct IKRIG_API FIKRetargetPoseOp : public FIKRetargetOpBase
+struct FIKRetargetPoseOp : public FIKRetargetOpBase
 {
 	GENERATED_BODY()
 	
-	virtual bool Initialize(
+	UE_API virtual bool Initialize(
 		const FIKRetargetProcessor& InProcessor,
 		const FRetargetSkeleton& InSourceSkeleton,
 		const FTargetSkeleton& InTargetSkeleton,
 		const FIKRetargetOpBase* InParentOp,
 		FIKRigLogger& Log) override;
 	
-	virtual void Run(
+	UE_API virtual void Run(
 		FIKRetargetProcessor& InProcessor,
 		const double InDeltaTime,
 		const TArray<FTransform>& InSourceGlobalPose,
 		TArray<FTransform>& OutTargetGlobalPose) override;
 
-	virtual FIKRetargetOpSettingsBase* GetSettings() override;
+	UE_API virtual FIKRetargetOpSettingsBase* GetSettings() override;
 
-	virtual void SetSettings(const FIKRetargetOpSettingsBase* InSettings) override;
+	UE_API virtual void SetSettings(const FIKRetargetOpSettingsBase* InSettings) override;
 	
-	virtual const UScriptStruct* GetSettingsType() const override;
+	UE_API virtual const UScriptStruct* GetSettingsType() const override;
 	
-	virtual const UScriptStruct* GetType() const override;
+	UE_API virtual const UScriptStruct* GetType() const override;
 
 	UPROPERTY()
 	FIKRetargetPoseOpSettings Settings;
 };
 
 /* The blueprint/python API for editing a Retarget Pose Op */
-UCLASS(BlueprintType)
-class IKRIG_API UIKRetargetPoseController : public UIKRetargetOpControllerBase
+UCLASS(MinimalAPI, BlueprintType)
+class UIKRetargetPoseController : public UIKRetargetOpControllerBase
 {
 	GENERATED_BODY()
 	
@@ -74,12 +76,14 @@ public:
 	/* Get the current op settings as a struct.
 	 * @return FIKRetargetPoseOpSettings struct with the current settings used by the op. */
 	UFUNCTION(BlueprintCallable, Category = Settings)
-	FIKRetargetPoseOpSettings GetSettings();
+	UE_API FIKRetargetPoseOpSettings GetSettings();
 
 	/* Set the op settings. Input is a custom struct type for this op.
 	 * @param InSettings a FIKRetargetPoseOpSettings struct containing all the settings to apply to this op */
 	UFUNCTION(BlueprintCallable, Category = Settings)
-	void SetSettings(FIKRetargetPoseOpSettings InSettings);
+	UE_API void SetSettings(FIKRetargetPoseOpSettings InSettings);
 };
 
 #undef LOCTEXT_NAMESPACE
+
+#undef UE_API
diff --git a/Engine/Plugins/Experimental/MeshResizing/Source/MeshResizingCore/Public/MeshResizing/BaseBodyTools.h b/Engine/Plugins/Experimental/MeshResizing/Source/MeshResizingCore/Public/MeshResizing/BaseBodyTools.h
--- a/Engine/Plugins/Experimental/MeshResizing/Source/MeshResizingCore/Public/MeshResizing/BaseBodyTools.h
+++ b/Engine/Plugins/Experimental/MeshResizing/Source/MeshResizingCore/Public/MeshResizing/BaseBodyTools.h
@@ -4,27 +4,30 @@
 #include "UObject/NameTypes.h"
 #include "Containers/Array.h"
 
+#define UE_API MESHRESIZINGCORE_API
+
 namespace UE::Geometry
 {
 	class FDynamicMesh3;
 }
 namespace UE::MeshResizing
 {
-	struct MESHRESIZINGCORE_API FBaseBodyTools
+	struct FBaseBodyTools
 	{
-		static const FName ImportedVertexVIDsAttrName;
-		static const FName RawPointIndicesVIDsAttrName;
+		static UE_API const FName ImportedVertexVIDsAttrName;
+		static UE_API const FName RawPointIndicesVIDsAttrName;
 
 		/** @return success */
-		static bool AttachVertexMappingData(const FName& AttrName, const TArray<int32>& Data, UE::Geometry::FDynamicMesh3& Mesh);
+		static UE_API bool AttachVertexMappingData(const FName& AttrName, const TArray<int32>& Data, UE::Geometry::FDynamicMesh3& Mesh);
 
 		/** @return success */
-		static bool GenerateResizableProxyFromVertexMappingData(const UE::Geometry::FDynamicMesh3& SourceMesh, const FName& SourceMappingName, const UE::Geometry::FDynamicMesh3& TargetMesh, const FName& TargetMappingName, UE::Geometry::FDynamicMesh3& ProxyMesh);
+		static UE_API bool GenerateResizableProxyFromVertexMappingData(const UE::Geometry::FDynamicMesh3& SourceMesh, const FName& SourceMappingName, const UE::Geometry::FDynamicMesh3& TargetMesh, const FName& TargetMappingName, UE::Geometry::FDynamicMesh3& ProxyMesh);
 
 		/** @return success */
-		static bool InterpolateResizableProxy(const UE::Geometry::FDynamicMesh3& SourceMesh, const UE::Geometry::FDynamicMesh3& TargetMesh, float BlendAlpha, UE::Geometry::FDynamicMesh3& ProxyMesh);
+		static UE_API bool InterpolateResizableProxy(const UE::Geometry::FDynamicMesh3& SourceMesh, const UE::Geometry::FDynamicMesh3& TargetMesh, float BlendAlpha, UE::Geometry::FDynamicMesh3& ProxyMesh);
 
 
 	};
 }
 
+#undef UE_API
diff --git a/Engine/Plugins/Experimental/MeshResizing/Source/MeshResizingCore/Public/MeshResizing/CustomRegionResizing.h b/Engine/Plugins/Experimental/MeshResizing/Source/MeshResizingCore/Public/MeshResizing/CustomRegionResizing.h
--- a/Engine/Plugins/Experimental/MeshResizing/Source/MeshResizingCore/Public/MeshResizing/CustomRegionResizing.h
+++ b/Engine/Plugins/Experimental/MeshResizing/Source/MeshResizingCore/Public/MeshResizing/CustomRegionResizing.h
@@ -8,6 +8,8 @@
 
 #include "CustomRegionResizing.generated.h"
 
+#define UE_API MESHRESIZINGCORE_API
+
 namespace UE::Geometry
 {
 	class FDynamicMesh3;
@@ -26,7 +28,7 @@ enum struct EMeshResizingCustomRegionType : uint8
 };
 
 USTRUCT()
-struct MESHRESIZINGCORE_API FMeshResizingCustomRegion
+struct FMeshResizingCustomRegion
 {
 	GENERATED_BODY()
 
@@ -72,11 +74,12 @@ struct MESHRESIZINGCORE_API FMeshResizingCustomRegion
 
 namespace UE::MeshResizing
 {
-	struct MESHRESIZINGCORE_API FCustomRegionResizing
+	struct FCustomRegionResizing
 	{
-		static void GenerateCustomRegion(const TConstArrayView<FVector3f>& BoundPositions, const UE::Geometry::FDynamicMesh3& SourceMesh, const TSet<int32>& BoundVertices, FMeshResizingCustomRegion& OutData);
-		static bool CalculateFrameForCustomRegion(const UE::Geometry::FDynamicMesh3& SourceMesh, const FMeshResizingCustomRegion& BindingGroup, FVector3d& OutOrigin, FVector3f& OutTangentU, FVector3f& OutTangentV, FVector3f& OutNormal);
-		static void InterpolateCustomRegionPoints(const FMeshResizingCustomRegion& BindingGroup, const TConstArrayView<FVector3d>& BoundsCorners, TArrayView<FVector3f> BoundPositions);
+		static UE_API void GenerateCustomRegion(const TConstArrayView<FVector3f>& BoundPositions, const UE::Geometry::FDynamicMesh3& SourceMesh, const TSet<int32>& BoundVertices, FMeshResizingCustomRegion& OutData);
+		static UE_API bool CalculateFrameForCustomRegion(const UE::Geometry::FDynamicMesh3& SourceMesh, const FMeshResizingCustomRegion& BindingGroup, FVector3d& OutOrigin, FVector3f& OutTangentU, FVector3f& OutTangentV, FVector3f& OutNormal);
+		static UE_API void InterpolateCustomRegionPoints(const FMeshResizingCustomRegion& BindingGroup, const TConstArrayView<FVector3d>& BoundsCorners, TArrayView<FVector3f> BoundPositions);
 	};
 }
 
+#undef UE_API
diff --git a/Engine/Plugins/Experimental/MeshResizing/Source/MeshResizingCore/Public/MeshResizing/RBFInterpolation.h b/Engine/Plugins/Experimental/MeshResizing/Source/MeshResizingCore/Public/MeshResizing/RBFInterpolation.h
--- a/Engine/Plugins/Experimental/MeshResizing/Source/MeshResizingCore/Public/MeshResizing/RBFInterpolation.h
+++ b/Engine/Plugins/Experimental/MeshResizing/Source/MeshResizingCore/Public/MeshResizing/RBFInterpolation.h
@@ -5,6 +5,8 @@
 #include "Containers/Array.h"
 #include "RBFInterpolation.generated.h"
 
+#define UE_API MESHRESIZINGCORE_API
+
 namespace UE::Geometry
 {
 	class FDynamicMesh3;
@@ -12,7 +14,7 @@ namespace UE::Geometry
 struct FMeshDescription;
 
 USTRUCT()
-struct MESHRESIZINGCORE_API FMeshResizingRBFInterpolationData
+struct FMeshResizingRBFInterpolationData
 {
 	GENERATED_BODY()
 
@@ -26,18 +28,19 @@ struct MESHRESIZINGCORE_API FMeshResizingRBFInterpolationData
 
 namespace UE::MeshResizing
 {
-	struct MESHRESIZINGCORE_API FRBFInterpolation
+	struct FRBFInterpolation
 	{
-		static void GenerateWeights(const TConstArrayView<FVector3f>& SourcePositions, int32 NumInterpolationPoints, FMeshResizingRBFInterpolationData& OutData);
-		static void GenerateWeights(const UE::Geometry::FDynamicMesh3& BaseMesh, int32 NumInterpolationPoints, FMeshResizingRBFInterpolationData& OutData);
-		static void GenerateMeshSamples(const TArray<FVector3f>& TargetPositions, const FMeshResizingRBFInterpolationData& InterpolationData, TArray<FVector3f>& SampleDeformations);
-		static void GenerateWeights(const FMeshDescription& BaseMesh, int32 NumInterpolationPoints, FMeshResizingRBFInterpolationData& OutData);
-		static void DeformPoints(const TArray<FVector3f>& TargetPositions, const FMeshResizingRBFInterpolationData& InterpolationData, bool bInterpolateNormals,  UE::Geometry::FDynamicMesh3& DeformingMesh);
-		static void DeformCoordinateFrames(const TArray<FVector3f>& TargetPositions, const FMeshResizingRBFInterpolationData& InterpolationData, bool bNormalize, bool bOrthogonalize, TArray<FMatrix44f>& Coordinates);
-		static void DeformPoints(const FMeshDescription& TargetMesh, const FMeshResizingRBFInterpolationData& InterpolationData, TArray<FVector3d>& Points);
-		static void DeformPoints(const FMeshDescription& TargetMesh, const FMeshResizingRBFInterpolationData& InterpolationData, bool bInterpolateNormals, UE::Geometry::FDynamicMesh3& DeformingMesh);
-		static void DeformPoints(const FMeshDescription& TargetMesh, const FMeshResizingRBFInterpolationData& InterpolationData, TArrayView<FVector3f> Points, TArrayView<FVector3f> Normals);
-		static void DeformPoints(const FMeshDescription& TargetMesh, const FMeshResizingRBFInterpolationData& InterpolationData, TArrayView<FVector3f> Points, TArrayView<FVector3f> Normals, TArrayView<FVector3f> TangentUs, TArrayView<FVector3f> TangentVs);
+		static UE_API void GenerateWeights(const TConstArrayView<FVector3f>& SourcePositions, int32 NumInterpolationPoints, FMeshResizingRBFInterpolationData& OutData);
+		static UE_API void GenerateWeights(const UE::Geometry::FDynamicMesh3& BaseMesh, int32 NumInterpolationPoints, FMeshResizingRBFInterpolationData& OutData);
+		static UE_API void GenerateMeshSamples(const TArray<FVector3f>& TargetPositions, const FMeshResizingRBFInterpolationData& InterpolationData, TArray<FVector3f>& SampleDeformations);
+		static UE_API void GenerateWeights(const FMeshDescription& BaseMesh, int32 NumInterpolationPoints, FMeshResizingRBFInterpolationData& OutData);
+		static UE_API void DeformPoints(const TArray<FVector3f>& TargetPositions, const FMeshResizingRBFInterpolationData& InterpolationData, bool bInterpolateNormals,  UE::Geometry::FDynamicMesh3& DeformingMesh);
+		static UE_API void DeformCoordinateFrames(const TArray<FVector3f>& TargetPositions, const FMeshResizingRBFInterpolationData& InterpolationData, bool bNormalize, bool bOrthogonalize, TArray<FMatrix44f>& Coordinates);
+		static UE_API void DeformPoints(const FMeshDescription& TargetMesh, const FMeshResizingRBFInterpolationData& InterpolationData, TArray<FVector3d>& Points);
+		static UE_API void DeformPoints(const FMeshDescription& TargetMesh, const FMeshResizingRBFInterpolationData& InterpolationData, bool bInterpolateNormals, UE::Geometry::FDynamicMesh3& DeformingMesh);
+		static UE_API void DeformPoints(const FMeshDescription& TargetMesh, const FMeshResizingRBFInterpolationData& InterpolationData, TArrayView<FVector3f> Points, TArrayView<FVector3f> Normals);
+		static UE_API void DeformPoints(const FMeshDescription& TargetMesh, const FMeshResizingRBFInterpolationData& InterpolationData, TArrayView<FVector3f> Points, TArrayView<FVector3f> Normals, TArrayView<FVector3f> TangentUs, TArrayView<FVector3f> TangentVs);
 	};
 }
 
+#undef UE_API
diff --git a/Engine/Plugins/Experimental/SubtitlesAndClosedCaptions/Source/SubtitlesAndClosedCaptions/Public/MovieSceneSubtitleSection.h b/Engine/Plugins/Experimental/SubtitlesAndClosedCaptions/Source/SubtitlesAndClosedCaptions/Public/MovieSceneSubtitleSection.h
--- a/Engine/Plugins/Experimental/SubtitlesAndClosedCaptions/Source/SubtitlesAndClosedCaptions/Public/MovieSceneSubtitleSection.h
+++ b/Engine/Plugins/Experimental/SubtitlesAndClosedCaptions/Source/SubtitlesAndClosedCaptions/Public/MovieSceneSubtitleSection.h
@@ -22,12 +22,12 @@ class UMovieSceneSubtitleSection
 
 public:
 
-	SUBTITLESANDCLOSEDCAPTIONS_API TObjectPtr<const USubtitleAssetUserData> GetSubtitle() const
+	TObjectPtr<const USubtitleAssetUserData> GetSubtitle() const
 	{
 		return Subtitle;
 	}
 
-	SUBTITLESANDCLOSEDCAPTIONS_API void SetSubtitle(const USubtitleAssetUserData& InSubtitle)
+	void SetSubtitle(const USubtitleAssetUserData& InSubtitle)
 	{
 		Subtitle = &InSubtitle;
 	}
diff --git a/Engine/Plugins/Experimental/SubtitlesAndClosedCaptions/Source/SubtitlesAndClosedCaptions/Public/MovieSceneSubtitlesSystem.h b/Engine/Plugins/Experimental/SubtitlesAndClosedCaptions/Source/SubtitlesAndClosedCaptions/Public/MovieSceneSubtitlesSystem.h
--- a/Engine/Plugins/Experimental/SubtitlesAndClosedCaptions/Source/SubtitlesAndClosedCaptions/Public/MovieSceneSubtitlesSystem.h
+++ b/Engine/Plugins/Experimental/SubtitlesAndClosedCaptions/Source/SubtitlesAndClosedCaptions/Public/MovieSceneSubtitlesSystem.h
@@ -17,7 +17,7 @@
 
 #include "MovieSceneSubtitlesSystem.generated.h"
 
-struct SUBTITLESANDCLOSEDCAPTIONS_API FEvaluateSubtitles
+struct FEvaluateSubtitles
 {
 
 	FEvaluateSubtitles(const UE::MovieScene::FInstanceRegistry& InInstanceRegistry)
diff --git a/Engine/Plugins/Experimental/SubtitlesAndClosedCaptions/Source/SubtitlesAndClosedCaptions/Public/SubtitleDataComponent.h b/Engine/Plugins/Experimental/SubtitlesAndClosedCaptions/Source/SubtitlesAndClosedCaptions/Public/SubtitleDataComponent.h
--- a/Engine/Plugins/Experimental/SubtitlesAndClosedCaptions/Source/SubtitlesAndClosedCaptions/Public/SubtitleDataComponent.h
+++ b/Engine/Plugins/Experimental/SubtitlesAndClosedCaptions/Source/SubtitlesAndClosedCaptions/Public/SubtitleDataComponent.h
@@ -12,7 +12,7 @@ class UMovieSceneSubtitleSection;
 
 /** Component data for subtitles tracks */
 USTRUCT()
-struct SUBTITLESANDCLOSEDCAPTIONS_API FSubtitleDataComponent
+struct FSubtitleDataComponent
 {
 	GENERATED_BODY()
 
diff --git a/Engine/Plugins/Experimental/SubtitlesAndClosedCaptions/Source/SubtitlesAndClosedCaptions/Public/SubtitlesBlueprintFunctionLibrary.h b/Engine/Plugins/Experimental/SubtitlesAndClosedCaptions/Source/SubtitlesAndClosedCaptions/Public/SubtitlesBlueprintFunctionLibrary.h
--- a/Engine/Plugins/Experimental/SubtitlesAndClosedCaptions/Source/SubtitlesAndClosedCaptions/Public/SubtitlesBlueprintFunctionLibrary.h
+++ b/Engine/Plugins/Experimental/SubtitlesAndClosedCaptions/Source/SubtitlesAndClosedCaptions/Public/SubtitlesBlueprintFunctionLibrary.h
@@ -6,24 +6,28 @@
 #include "Subtitles/SubtitlesAndClosedCaptionsDelegates.h"
 #include "SubtitlesBlueprintFunctionLibrary.generated.h"
 
+#define UE_API SUBTITLESANDCLOSEDCAPTIONS_API
+
 class USubtitleAssetUserData;
 
-UCLASS(meta = (ScriptName = "SubtitlesLibrary"))
-class SUBTITLESANDCLOSEDCAPTIONS_API USubtitlesBlueprintFunctionLibrary : public UBlueprintFunctionLibrary
+UCLASS(MinimalAPI, meta = (ScriptName = "SubtitlesLibrary"))
+class USubtitlesBlueprintFunctionLibrary : public UBlueprintFunctionLibrary
 {
 	GENERATED_BODY()
 
 public:
 
 	UFUNCTION(BlueprintCallable, Category = Subtitles)
-	static void QueueSubtitle(const USubtitleAssetUserData* Subtitle, const ESubtitleTiming Timing = ESubtitleTiming::InternallyTimed);
+	static UE_API void QueueSubtitle(const USubtitleAssetUserData* Subtitle, const ESubtitleTiming Timing = ESubtitleTiming::InternallyTimed);
 
 	UFUNCTION(BlueprintCallable, Category = Subtitles)
-	static bool IsSubtitleActive(const USubtitleAssetUserData* Subtitle);
+	static UE_API bool IsSubtitleActive(const USubtitleAssetUserData* Subtitle);
 
 	UFUNCTION(BlueprintCallable, Category = Subtitles)
-	static void StopSubtitle(const USubtitleAssetUserData* Subtitle);
+	static UE_API void StopSubtitle(const USubtitleAssetUserData* Subtitle);
 
 	UFUNCTION(BlueprintCallable, Category = Subtitles)
-	static void StopAllSubtitles();
-};
\ No newline at end of file
+	static UE_API void StopAllSubtitles();
+};
+
+#undef UE_API
diff --git a/Engine/Plugins/Experimental/SubtitlesAndClosedCaptions/Source/SubtitlesAndClosedCaptions/Public/SubtitlesSubsystem.h b/Engine/Plugins/Experimental/SubtitlesAndClosedCaptions/Source/SubtitlesAndClosedCaptions/Public/SubtitlesSubsystem.h
--- a/Engine/Plugins/Experimental/SubtitlesAndClosedCaptions/Source/SubtitlesAndClosedCaptions/Public/SubtitlesSubsystem.h
+++ b/Engine/Plugins/Experimental/SubtitlesAndClosedCaptions/Source/SubtitlesAndClosedCaptions/Public/SubtitlesSubsystem.h
@@ -12,6 +12,8 @@
 #include "Subtitles/SubtitlesAndClosedCaptionsDelegates.h"
 #include "SubtitlesSubsystem.generated.h"
 
+#define UE_API SUBTITLESANDCLOSEDCAPTIONS_API
+
 class FCanvas;
 class UAssetUserData;
 class UFont;
@@ -33,8 +35,8 @@ namespace SubtitlesAndClosedCaptions::Test
 *
 *	Game configuration for font customization per game
 */
-UCLASS(config = Game, defaultConfig)
-class SUBTITLESANDCLOSEDCAPTIONS_API USubtitlesSubsystem : public UWorldSubsystem
+UCLASS(MinimalAPI, config = Game, defaultConfig)
+class USubtitlesSubsystem : public UWorldSubsystem
 {
 	GENERATED_BODY()
 public:
@@ -45,21 +47,21 @@ public:
 	// Adds a subtitle to the queue: Params contains the subtitle asset and an optional duration. The highest-priority subtitle in the queue will be displayed.
 	// If Timing is ExternallyTimed, the queued subtitle will remain in the queue until manually removed.
 	// If the subtitle asset has a non-zero StartOffset, it will sit in a delayed-start queue instead of being queued for display.
-	virtual void QueueSubtitle(const FQueueSubtitleParameters& Params, const ESubtitleTiming Timing = ESubtitleTiming::InternallyTimed);
+	UE_API virtual void QueueSubtitle(const FQueueSubtitleParameters& Params, const ESubtitleTiming Timing = ESubtitleTiming::InternallyTimed);
 
 	// Returns true if the given subtitle asset is being displayed.
-	virtual bool IsSubtitleActive(const UAssetUserData& Data) const;
+	UE_API virtual bool IsSubtitleActive(const UAssetUserData& Data) const;
 
 	// Stops the given subtitle asset being displayed.  This includes subtitles not yet being displayed due to their StartOffset
-	virtual void StopSubtitle(const UAssetUserData& Data);
+	UE_API virtual void StopSubtitle(const UAssetUserData& Data);
 
 	// Stops all queued subtitles from being displayed.  This includes subtitles not yet being displayed due to their StartOffset.
-	virtual void StopAllSubtitles();
+	UE_API virtual void StopAllSubtitles();
 
 protected:
 
-	virtual void Initialize(FSubsystemCollectionBase& Collection) override;
-	virtual void BindDelegates();
+	UE_API virtual void Initialize(FSubsystemCollectionBase& Collection) override;
+	UE_API virtual void BindDelegates();
 
 	UPROPERTY()
 	TMap<ESubtitleType, FSlateFontInfo> SubtitleFontInfo;
@@ -73,9 +75,9 @@ protected:
 	TArray<FActiveSubtitle> DelayedSubtitles;
 
 protected:
-	virtual void AddActiveSubtitle(const USubtitleAssetUserData& Subtitle, float Duration);
-	virtual void MakeDelayedSubtitleActive(TObjectPtr<const USubtitleAssetUserData> Subtitle);
-	virtual void RemoveActiveSubtitle(TObjectPtr<const USubtitleAssetUserData> Subtitle);
+	UE_API virtual void AddActiveSubtitle(const USubtitleAssetUserData& Subtitle, float Duration);
+	UE_API virtual void MakeDelayedSubtitleActive(TObjectPtr<const USubtitleAssetUserData> Subtitle);
+	UE_API virtual void RemoveActiveSubtitle(TObjectPtr<const USubtitleAssetUserData> Subtitle);
 
 #if WITH_DEV_AUTOMATION_TESTS
 #if 0 // Temporarily disabling these tests as they have a dangling reference that trips up the static analysis on certain build configurations.
@@ -110,6 +112,8 @@ private:
 	TObjectPtr<USubtitleTextBlock> SubtitleWidget = nullptr;
 
 	bool bInitializedWidget = false;
-	bool TryCreateUMGWidget();
-	void UpdateWidgetData();
+	UE_API bool TryCreateUMGWidget();
+	UE_API void UpdateWidgetData();
 };
+
+#undef UE_API
diff --git a/Engine/Plugins/Experimental/Water/Source/Runtime/Public/GerstnerWaterWaveSubsystem.h b/Engine/Plugins/Experimental/Water/Source/Runtime/Public/GerstnerWaterWaveSubsystem.h
--- a/Engine/Plugins/Experimental/Water/Source/Runtime/Public/GerstnerWaterWaveSubsystem.h
+++ b/Engine/Plugins/Experimental/Water/Source/Runtime/Public/GerstnerWaterWaveSubsystem.h
@@ -25,7 +25,7 @@ public:
 	WATER_API void Register(FWaterViewExtension* ViewExtension);
 	WATER_API void Unregister(FWaterViewExtension* ViewExtension);
 
-	WATER_API void RebuildGPUData() { bRebuildGPUData = true; }
+	void RebuildGPUData() { bRebuildGPUData = true; }
 
 	//~ Begin UObject Interface.	
 	WATER_API virtual void GetResourceSizeEx(FResourceSizeEx& CumulativeResourceSize) override;
diff --git a/Engine/Plugins/Experimental/Water/Source/Runtime/Public/GerstnerWaterWaves.h b/Engine/Plugins/Experimental/Water/Source/Runtime/Public/GerstnerWaterWaves.h
--- a/Engine/Plugins/Experimental/Water/Source/Runtime/Public/GerstnerWaterWaves.h
+++ b/Engine/Plugins/Experimental/Water/Source/Runtime/Public/GerstnerWaterWaves.h
@@ -177,7 +177,7 @@ public:
 	UGerstnerWaterWaves();
 	
 	/** Returns the maximum wave height that can be reached by those waves */
-	WATER_API virtual float GetMaxWaveHeight() const override { return MaxWaveHeight; }
+	virtual float GetMaxWaveHeight() const override { return MaxWaveHeight; }
 
 	/** Computes the raw wave perturbation of the water height/normal */
 	WATER_API virtual float GetWaveHeightAtPosition(const FVector& InPosition, float InWaterDepth, float InTime, FVector& OutNormal) const override;
@@ -200,7 +200,7 @@ protected:
 	float MaxWaveHeight;
 
 public:
-	WATER_API const TArray<FGerstnerWave>& GetGerstnerWaves() const { return GerstnerWaves; }
+	const TArray<FGerstnerWave>& GetGerstnerWaves() const { return GerstnerWaves; }
 
 	virtual void PostDuplicate(EDuplicateMode::Type DuplicateMode) override;
 #if WITH_EDITOR
diff --git a/Engine/Plugins/FX/Niagara/Source/Niagara/Classes/NiagaraValidationRule.h b/Engine/Plugins/FX/Niagara/Source/Niagara/Classes/NiagaraValidationRule.h
--- a/Engine/Plugins/FX/Niagara/Source/Niagara/Classes/NiagaraValidationRule.h
+++ b/Engine/Plugins/FX/Niagara/Source/Niagara/Classes/NiagaraValidationRule.h
@@ -70,7 +70,7 @@ class UNiagaraValidationRule : public UObject
 	GENERATED_BODY()
 
 public:
-	NIAGARA_API bool IsEnabled() const { return bIsConfigDisabled == false; }
+	bool IsEnabled() const { return bIsConfigDisabled == false; }
 #if WITH_EDITOR
 	NIAGARA_API virtual void CheckValidity(const FNiagaraValidationContext& Context, TArray<FNiagaraValidationResult>& OutResults) const;
 #endif
diff --git a/Engine/Plugins/Online/OnlineServices/Source/OnlineServicesCommon/Public/Online/OnlineServicesCommonDelegates.h b/Engine/Plugins/Online/OnlineServices/Source/OnlineServicesCommon/Public/Online/OnlineServicesCommonDelegates.h
--- a/Engine/Plugins/Online/OnlineServices/Source/OnlineServicesCommon/Public/Online/OnlineServicesCommonDelegates.h
+++ b/Engine/Plugins/Online/OnlineServices/Source/OnlineServicesCommon/Public/Online/OnlineServicesCommonDelegates.h
@@ -5,12 +5,14 @@
 #include "Delegates/Delegate.h"
 #include "Online/OnlineError.h"
 
+#define UE_API ONLINESERVICESCOMMON_API
+
 namespace UE::Online {
 
 class FOnlineServicesCommon;
 
 // Notification params of OnOnlineAsyncOpCompletedV2 delegate
-struct ONLINESERVICESCOMMON_API FOnOnlineAsyncOpCompletedParams
+struct FOnOnlineAsyncOpCompletedParams
 {
 	/**
 	 * FOnOnlineAsyncOpCompletedParams ctor
@@ -18,7 +20,7 @@ struct ONLINESERVICESCOMMON_API FOnOnlineAsyncOpCompletedParams
 	 * @param OnlineServicesCommon the OnlineServices instance
 	 * @param OnlineError the result of completed operation
 	 */
-	FOnOnlineAsyncOpCompletedParams(FOnlineServicesCommon& InOnlineServicesCommon, const TOptional<FOnlineError>& InOnlineError);
+	UE_API FOnOnlineAsyncOpCompletedParams(FOnlineServicesCommon& InOnlineServicesCommon, const TOptional<FOnlineError>& InOnlineError);
 
 	// The name of completed operation
 	FString OpName;
@@ -43,3 +45,5 @@ DECLARE_MULTICAST_DELEGATE_OneParam(FOnOnlineAsyncOpCompletedV2, const FOnOnline
 extern ONLINESERVICESCOMMON_API FOnOnlineAsyncOpCompletedV2 OnOnlineAsyncOpCompletedV2;
 
 /* UE::Online */ }
+
+#undef UE_API
diff --git a/Engine/Plugins/Runtime/HairStrands/Source/HairStrandsCore/Public/HairStrandsDatas.h b/Engine/Plugins/Runtime/HairStrands/Source/HairStrandsCore/Public/HairStrandsDatas.h
--- a/Engine/Plugins/Runtime/HairStrands/Source/HairStrandsCore/Public/HairStrandsDatas.h
+++ b/Engine/Plugins/Runtime/HairStrands/Source/HairStrandsCore/Public/HairStrandsDatas.h
@@ -608,7 +608,7 @@ struct FHairStrandsCommonDatas
 
 	UE_API uint32 GetAttributes() const;
 	UE_API uint32 GetAttributeFlags() const;
-	UE_API bool HasFlags(EHairStrandsFlags In) const { return !!(Flags & uint32(In)); }
+	bool HasFlags(EHairStrandsFlags In) const { return !!(Flags & uint32(In)); }
 
 	UE_API void Reset();
 
@@ -1054,4 +1054,4 @@ struct FHairStrandsDebugResources
 	TRefCountPtr<FRDGPooledBuffer> VoxelData;
 };
 
-#undef UE_API
\ No newline at end of file
+#undef UE_API
diff --git a/Engine/Plugins/Runtime/Metasound/Source/MetasoundFrontend/Public/MetasoundFrontendDocument.h b/Engine/Plugins/Runtime/Metasound/Source/MetasoundFrontend/Public/MetasoundFrontendDocument.h
--- a/Engine/Plugins/Runtime/Metasound/Source/MetasoundFrontend/Public/MetasoundFrontendDocument.h
+++ b/Engine/Plugins/Runtime/Metasound/Source/MetasoundFrontend/Public/MetasoundFrontendDocument.h
@@ -195,7 +195,7 @@ struct FMetasoundCommentNodeIntVector : public FIntVector2
 {
 	GENERATED_BODY()
 
-	UE_API FMetasoundCommentNodeIntVector() = default;
+	FMetasoundCommentNodeIntVector() = default;
 	UE_API FMetasoundCommentNodeIntVector(const FIntVector2& InValue);
 	UE_API FMetasoundCommentNodeIntVector(const FVector2f& InValue);
 	UE_API FMetasoundCommentNodeIntVector(const FVector2d& InValue);
diff --git a/Engine/Source/Runtime/Experimental/Dataflow/Core/Public/Dataflow/DataflowContextAssetStore.h b/Engine/Source/Runtime/Experimental/Dataflow/Core/Public/Dataflow/DataflowContextAssetStore.h
--- a/Engine/Source/Runtime/Experimental/Dataflow/Core/Public/Dataflow/DataflowContextAssetStore.h
+++ b/Engine/Source/Runtime/Experimental/Dataflow/Core/Public/Dataflow/DataflowContextAssetStore.h
@@ -77,7 +77,7 @@ namespace UE::Dataflow
 
 		//~ Begin FGCObject interface
 		DATAFLOWCORE_API virtual void AddReferencedObjects(FReferenceCollector& Collector) override;
-		DATAFLOWCORE_API virtual FString GetReferencerName() const override { return TEXT("FContextAssetStore"); }
+		virtual FString GetReferencerName() const override { return TEXT("FContextAssetStore"); }
 		//~ End FGCObject interface
 
 	private:
diff --git a/Engine/Source/Runtime/Experimental/Iris/Core/Public/Iris/ReplicationSystem/ReplicationBridge.h b/Engine/Source/Runtime/Experimental/Iris/Core/Public/Iris/ReplicationSystem/ReplicationBridge.h
--- a/Engine/Source/Runtime/Experimental/Iris/Core/Public/Iris/ReplicationSystem/ReplicationBridge.h
+++ b/Engine/Source/Runtime/Experimental/Iris/Core/Public/Iris/ReplicationSystem/ReplicationBridge.h
@@ -247,15 +247,15 @@ protected:
 	IRISCORE_API virtual bool IsAllowedToDestroyInstance(const UObject* Instance) const;
 
 	/** Called when a remote connection detected a protocol mismatch when trying to instantiate the NetRefHandle replicated object. */
-	IRISCORE_API virtual void OnProtocolMismatchReported(FNetRefHandle RefHandle, uint32 ConnectionId) {}
+	virtual void OnProtocolMismatchReported(FNetRefHandle RefHandle, uint32 ConnectionId) {}
 
 	/** Called when a remote connection has a critical error caused by a specific NetRefHandle */
-	IRISCORE_API virtual void OnErrorWithNetRefHandleReported(UE::Net::ENetRefHandleError ErrorType, FNetRefHandle RefHandle, uint32 ConnectionId) {}
-	IRISCORE_API virtual void OnErrorWithNetRefHandleReported(UE::Net::ENetRefHandleError ErrorType, FNetRefHandle RefHandle, uint32 ConnectionId, const TArray<FNetRefHandle>& ExtraNetRefHandle) {}
+	virtual void OnErrorWithNetRefHandleReported(UE::Net::ENetRefHandleError ErrorType, FNetRefHandle RefHandle, uint32 ConnectionId) {}
+	virtual void OnErrorWithNetRefHandleReported(UE::Net::ENetRefHandleError ErrorType, FNetRefHandle RefHandle, uint32 ConnectionId, const TArray<FNetRefHandle>& ExtraNetRefHandle) {}
 
 	/** Tell the remote connection that we detected a reading error with a specific replicated object */
-	IRISCORE_API virtual void SendErrorWithNetRefHandle(UE::Net::ENetRefHandleError ErrorType, FNetRefHandle RefHandle, uint32 ConnectionId) {}
-	IRISCORE_API virtual void SendErrorWithNetRefHandle(UE::Net::ENetRefHandleError ErrorType, FNetRefHandle RefHandle, uint32 ConnectionId, const TArray<FNetRefHandle>& ExtraNetRefHandle) {}
+	virtual void SendErrorWithNetRefHandle(UE::Net::ENetRefHandleError ErrorType, FNetRefHandle RefHandle, uint32 ConnectionId) {}
+	virtual void SendErrorWithNetRefHandle(UE::Net::ENetRefHandleError ErrorType, FNetRefHandle RefHandle, uint32 ConnectionId, const TArray<FNetRefHandle>& ExtraNetRefHandle) {}
 
 private:
 
diff --git a/Engine/Source/Runtime/Interchange/Core/Public/Nodes/InterchangeBaseNodeContainer.h b/Engine/Source/Runtime/Interchange/Core/Public/Nodes/InterchangeBaseNodeContainer.h
--- a/Engine/Source/Runtime/Interchange/Core/Public/Nodes/InterchangeBaseNodeContainer.h
+++ b/Engine/Source/Runtime/Interchange/Core/Public/Nodes/InterchangeBaseNodeContainer.h
@@ -55,7 +55,7 @@ public:
 	 * Removes node from Nodes map with the given NodeUniqueID.
 	 */
 	UFUNCTION(BlueprintCallable, Category = "Interchange | Node Container")
-	INTERCHANGECORE_API void RemoveNode(const FString& NodeUniqueID)
+	void RemoveNode(const FString& NodeUniqueID)
 	{
 		Nodes.Remove(NodeUniqueID);
 	}
diff --git a/Engine/Source/Runtime/MovieSceneTracks/Public/MovieSceneTracksSettings.h b/Engine/Source/Runtime/MovieSceneTracks/Public/MovieSceneTracksSettings.h
--- a/Engine/Source/Runtime/MovieSceneTracks/Public/MovieSceneTracksSettings.h
+++ b/Engine/Source/Runtime/MovieSceneTracks/Public/MovieSceneTracksSettings.h
@@ -7,15 +7,17 @@
 
 #include "MovieSceneTracksSettings.generated.h"
 
+#define UE_API MOVIESCENETRACKS_API
+
 /** Options for some of the Sequencer systems in this module. */
-UCLASS(config=EditorPerProjectUserSettings, PerObjectConfig)
-class MOVIESCENETRACKS_API UMovieSceneTracksSettings : public UObject
+UCLASS(MinimalAPI, config=EditorPerProjectUserSettings, PerObjectConfig)
+class UMovieSceneTracksSettings : public UObject
 {
 	GENERATED_BODY()
 
 public:
 
-	UMovieSceneTracksSettings(const FObjectInitializer& ObjectInitializer);
+	UE_API UMovieSceneTracksSettings(const FObjectInitializer& ObjectInitializer);
 
 	/**
 	 * Gets whether camera cut tracks should take control of the viewport in SIE, or PIE while ejected from the player controller.
@@ -25,7 +27,7 @@ public:
 	/**
 	 * Sets whether camera cut tracks should take control of the viewport in SIE, or PIE while ejected from the player controller.
 	 */
-	void SetPreviewCameraCutsInSimulate(bool bInPreviewCameraCutsInSimulate);
+	UE_API void SetPreviewCameraCutsInSimulate(bool bInPreviewCameraCutsInSimulate);
 
 protected:
 
@@ -37,3 +39,4 @@ protected:
 	bool bPreviewCameraCutsInSimulate;
 };
 
+#undef UE_API
diff --git a/Engine/Source/Runtime/StorageServerClient/Private/ZenStreamingSettings.h b/Engine/Source/Runtime/StorageServerClient/Private/ZenStreamingSettings.h
--- a/Engine/Source/Runtime/StorageServerClient/Private/ZenStreamingSettings.h
+++ b/Engine/Source/Runtime/StorageServerClient/Private/ZenStreamingSettings.h
@@ -7,8 +7,8 @@
 
 #include "ZenStreamingSettings.generated.h"
 
-UCLASS(config=Game, defaultconfig, meta=(DisplayName="Zen Streaming"))
-class STORAGESERVERCLIENT_API UZenStreamingSettings : public UDeveloperSettings
+UCLASS(MinimalAPI, config=Game, defaultconfig, meta=(DisplayName="Zen Streaming"))
+class UZenStreamingSettings : public UDeveloperSettings
 {
 	GENERATED_BODY()
 
@@ -33,4 +33,4 @@ public:
 
 	UPROPERTY(config, EditAnywhere, Category = ZenStreaming, meta=(ClampMin=0, ClampMax=1.0, ConsoleVariable="zen.indicator.alpha"))
 	float BadgeAlpha = 0.5;
-};
\ No newline at end of file
+};
