Mesh Resizing - Fixed race condition when creating material instance.

#rnx

diff --git a/Engine/Plugins/Experimental/ChaosOutfitAsset/Source/ChaosOutfitAssetDataflowNodes/Private/ChaosOutfitAsset/OutfitAssetTerminalNode.cpp b/Engine/Plugins/Experimental/ChaosOutfitAsset/Source/ChaosOutfitAssetDataflowNodes/Private/ChaosOutfitAsset/OutfitAssetTerminalNode.cpp
--- a/Engine/Plugins/Experimental/ChaosOutfitAsset/Source/ChaosOutfitAssetDataflowNodes/Private/ChaosOutfitAsset/OutfitAssetTerminalNode.cpp
+++ b/Engine/Plugins/Experimental/ChaosOutfitAsset/Source/ChaosOutfitAssetDataflowNodes/Private/ChaosOutfitAsset/OutfitAssetTerminalNode.cpp
@@ -105,7 +105,7 @@ void FChaosOutfitAssetTerminalNode::SetAssetValue(TObjectPtr<UObject> Asset, UE:
 		const TObjectPtr<const UChaosOutfit>& InOutfit = GetValue(Context, &Outfit);
 
 		// Build the asset
-		OutfitAsset->Build(InOutfit);
+		OutfitAsset->Build(InOutfit, &Context);
 
 		// Asset must be resaved
 		OutfitAsset->MarkPackageDirty();
diff --git a/Engine/Plugins/Experimental/ChaosOutfitAsset/Source/ChaosOutfitAssetEngine/Private/ChaosOutfitAsset/OutfitAsset.cpp b/Engine/Plugins/Experimental/ChaosOutfitAsset/Source/ChaosOutfitAssetEngine/Private/ChaosOutfitAsset/OutfitAsset.cpp
--- a/Engine/Plugins/Experimental/ChaosOutfitAsset/Source/ChaosOutfitAssetEngine/Private/ChaosOutfitAsset/OutfitAsset.cpp
+++ b/Engine/Plugins/Experimental/ChaosOutfitAsset/Source/ChaosOutfitAssetEngine/Private/ChaosOutfitAsset/OutfitAsset.cpp
@@ -2,8 +2,10 @@
 
 #include "ChaosOutfitAsset/OutfitAsset.h"
 #include "ChaosClothAsset/ClothSimulationModel.h"
+#include "ChaosClothAsset/CollectionClothFacade.h"
 #include "ChaosOutfitAsset/CollectionOutfitFacade.h"
 #include "ChaosOutfitAsset/OutfitAssetPrivate.h"
+#include "Dataflow/DataflowContextAssetStore.h"
 #include "Engine/SkeletalMesh.h"
 #if WITH_EDITOR
 #include "DerivedDataCacheInterface.h"
@@ -46,7 +48,7 @@ bool UChaosOutfitAsset::HasValidClothSimulationModels() const
 	return false;
 }
 
-void UChaosOutfitAsset::Build(const TObjectPtr<const UChaosOutfit> InOutfit)
+void UChaosOutfitAsset::Build(const TObjectPtr<const UChaosOutfit> InOutfit, UE::Dataflow::IContextAssetStoreInterface* ContextAssetStore)
 {
 	using namespace UE::Chaos::OutfitAsset;
 
@@ -73,6 +75,81 @@ void UChaosOutfitAsset::Build(const TObjectPtr<const UChaosOutfit> InOutfit)
 			Outfit = DuplicateObject<UChaosOutfit>(InOutfit, this, UniqueOutfitName);
 		}
 #endif
+		if (ContextAssetStore)
+		{
+			// Fix up resized material paths
+			TMap<FString, FString> MaterialPathsToFixUp;
+			MaterialPathsToFixUp.Reserve(Materials.Num());
+
+			for (int32 MaterialIndex = 0; MaterialIndex < Materials.Num(); ++MaterialIndex)
+			{
+				FSkeletalMaterial& Material = Materials[MaterialIndex];
+				if (Material.MaterialInterface->GetOuter() == GetTransientPackage())
+				{
+					const FString TransientPathName = Material.MaterialInterface->GetPathName();
+					Material.MaterialInterface = Cast<UMaterialInterface>(ContextAssetStore->CommitAsset(TransientPathName));
+#if WITH_EDITORONLY_DATA
+					if (Outfit)
+					{
+						Outfit->GetMaterials()[MaterialIndex].MaterialInterface = Material.MaterialInterface;
+					}
+#endif
+					MaterialPathsToFixUp.Emplace(TransientPathName, Material.MaterialInterface ? Material.MaterialInterface->GetPathName() : FString());
+				}
+			}
+
+			// Fix up cloth collections
+			if (MaterialPathsToFixUp.Num())
+			{
+				auto FixUpPiecesMaterials = [&MaterialPathsToFixUp](TArrayView<FChaosOutfitPiece> Pieces)
+					{
+						using namespace UE::Chaos::ClothAsset;
+						for (FChaosOutfitPiece& Piece : Pieces)
+						{
+							for (TSharedRef<const FManagedArrayCollection>& Collection : Piece.Collections)
+							{
+								TOptional<TSharedRef<FManagedArrayCollection>> FixedUpCollection;
+								TOptional<FCollectionClothFacade> FixedUpClothFacade;
+
+								FCollectionClothConstFacade ClothFacade(Collection);
+								const TConstArrayView<FString> RenderMaterialPathNames = ClothFacade.GetRenderMaterialPathName();
+								for (int32 PathIndex = 0; PathIndex < RenderMaterialPathNames.Num(); ++PathIndex)
+								{
+									const FString& RenderMaterialPathName = RenderMaterialPathNames[PathIndex];
+									if (const FString* MaterialPathToFixUp = MaterialPathsToFixUp.Find(RenderMaterialPathName))
+									{
+										if (!FixedUpCollection.IsSet())
+										{
+											FixedUpCollection.Emplace(MakeShared<FManagedArrayCollection>(*Collection));
+											FixedUpClothFacade.Emplace(FixedUpCollection.GetValue());
+										}
+										TArrayView<FString> FixedUpRenderMaterialPathNames = FixedUpClothFacade->GetRenderMaterialPathName();
+										check(FixedUpRenderMaterialPathNames[PathIndex] == RenderMaterialPathName);
+							
+										FixedUpRenderMaterialPathNames[PathIndex] = *MaterialPathToFixUp;
+									}
+								}
+
+								if (FixedUpCollection.IsSet())
+								{
+									Collection = MoveTemp(FixedUpCollection.GetValue());
+								}
+							}
+						}
+					};
+
+				// Fix up this outfit's pieces
+				FixUpPiecesMaterials(Pieces);
+
+#if WITH_EDITORONLY_DATA
+				// Fix up this outfit construction's pieces
+				if (Outfit)
+				{
+					FixUpPiecesMaterials(Outfit->GetPieces());
+				}
+#endif
+			}
+		}
 	}
 	else
 	{
diff --git a/Engine/Plugins/Experimental/ChaosOutfitAsset/Source/ChaosOutfitAssetEngine/Public/ChaosOutfitAsset/Outfit.h b/Engine/Plugins/Experimental/ChaosOutfitAsset/Source/ChaosOutfitAssetEngine/Public/ChaosOutfitAsset/Outfit.h
--- a/Engine/Plugins/Experimental/ChaosOutfitAsset/Source/ChaosOutfitAssetEngine/Public/ChaosOutfitAsset/Outfit.h
+++ b/Engine/Plugins/Experimental/ChaosOutfitAsset/Source/ChaosOutfitAssetEngine/Public/ChaosOutfitAsset/Outfit.h
@@ -112,6 +112,16 @@ public:
 		return Pieces;
 	}
 
+	const TArray<FSkeletalMaterial>& GetMaterials() const
+	{
+		return Materials;
+	}
+
+	TArray<FSkeletalMaterial>& GetMaterials()
+	{
+		return Materials;
+	}
+
 	const FManagedArrayCollection& GetOutfitCollection() const
 	{
 		return OutfitCollection;
diff --git a/Engine/Plugins/Experimental/ChaosOutfitAsset/Source/ChaosOutfitAssetEngine/Public/ChaosOutfitAsset/OutfitAsset.h b/Engine/Plugins/Experimental/ChaosOutfitAsset/Source/ChaosOutfitAssetEngine/Public/ChaosOutfitAsset/OutfitAsset.h
--- a/Engine/Plugins/Experimental/ChaosOutfitAsset/Source/ChaosOutfitAssetEngine/Public/ChaosOutfitAsset/OutfitAsset.h
+++ b/Engine/Plugins/Experimental/ChaosOutfitAsset/Source/ChaosOutfitAssetEngine/Public/ChaosOutfitAsset/OutfitAsset.h
@@ -10,6 +10,7 @@
 
 struct FChaosOutfitPiece;
 struct FSkeletalMaterial;
+namespace UE::Dataflow { struct IContextAssetStoreInterface; }
 
 /**
  * Outfit asset for character clothing and simulation.
@@ -23,7 +24,7 @@ public:
 	UE_API UChaosOutfitAsset(FVTableHelper& Helper);  // This is declared so we can use TUniquePtr<FClothSimulationModel> with just a forward declare of that class
 	UE_API virtual ~UChaosOutfitAsset() override;
 
-	UE_API void Build(const TObjectPtr<const UChaosOutfit> InOutfit);
+	UE_API void Build(const TObjectPtr<const UChaosOutfit> InOutfit, UE::Dataflow::IContextAssetStoreInterface* ContextAssetStore = nullptr);
 
 	//~ Begin UObject interface
 	UE_API virtual void Serialize(FArchive& Ar) override;
diff --git a/Engine/Plugins/Experimental/MeshResizing/Source/MeshResizingNodes/Private/MeshResizing/UVResizeControllerNode.cpp b/Engine/Plugins/Experimental/MeshResizing/Source/MeshResizingNodes/Private/MeshResizing/UVResizeControllerNode.cpp
--- a/Engine/Plugins/Experimental/MeshResizing/Source/MeshResizingNodes/Private/MeshResizing/UVResizeControllerNode.cpp
+++ b/Engine/Plugins/Experimental/MeshResizing/Source/MeshResizingNodes/Private/MeshResizing/UVResizeControllerNode.cpp
@@ -1,6 +1,5 @@
 // Copyright Epic Games, Inc. All Rights Reserved.
 #include "MeshResizing/UVResizeControllerNode.h"
-#include "AssetRegistry/AssetRegistryModule.h"
 #include "Dataflow/DataflowMesh.h"
 #include "DynamicMesh/DynamicMeshAttributeSet.h"
 #include "Engine/Texture.h"
@@ -128,61 +127,104 @@ void FUVResizeControllerNode::Evaluate(UE::Dataflow::FContext& Context, const FD
 						{
 							check(Material);  // Already tested in FindUVChannelParameterNames
 
-							// Create a new material instance that has this Material as base
-							const FString MaterialPackagePath = FPackageName::GetLongPackagePath(Material->GetOutermost()->GetName());
-							FString MaterialName = Material->GetName() + TEXT("_UVResized");
-							FString MaterialPackageName = FPaths::Combine(MaterialPackagePath, MaterialName);
-							if (FindPackage(nullptr, *MaterialPackageName))
-							{
-								MaterialPackageName = MakeUniqueObjectName(nullptr, UPackage::StaticClass(), FName(MaterialPackageName)).ToString();
-								MaterialName = FPaths::GetBaseFilename(MaterialPackageName);
-							}
-							if (UPackage* const MaterialPackage = CreatePackage(*MaterialPackageName))
+							// Re-assign the UVChannels in this new material instance
+							TArray<float> UVChannelParameterIndices;
+							UVChannelParameterIndices.SetNumZeroed(UVChannelParameterNames.Num());
+
+							for (int32 ParameterIndex = 0; ParameterIndex < UVChannelParameterNames.Num(); ++ParameterIndex)
 							{
-								if (UMaterialInstanceConstant* const MaterialInstance = NewObject<UMaterialInstanceConstant>(MaterialPackage, FName(MaterialName), RF_Public | RF_Standalone))
+								const FString& UVChannelParameterName = UVChannelParameterNames[ParameterIndex];
+
+								float Value;
+								if (ensure(Material->GetScalarParameterValue(*UVChannelParameterName, Value)))
 								{
-									MaterialInstance->PreEditChange(nullptr);
-									MaterialInstance->PostEditChange();  // Call to InitResources
+									const int32 SourceUVChannelIndex = (int32)Value;
+									if (UVChannelRemaps.IsValidIndex(SourceUVChannelIndex))
+									{
+										int32& UVChannelIndex = UVChannelRemaps[SourceUVChannelIndex];
+										if (UVChannelIndex == INDEX_NONE)
+										{
+											// Duplicate UV channel
+											UVChannelIndex = DynamicMesh.Attributes()->NumUVLayers();
+											DynamicMesh.Attributes()->SetNumUVLayers(UVChannelIndex + 1);
+											const UE::Geometry::FDynamicMeshUVOverlay* const SourceUVLayer = DynamicMesh.Attributes()->GetUVLayer(SourceUVChannelIndex);
+											check(SourceUVLayer);
+											DynamicMesh.Attributes()->GetUVLayer(UVChannelIndex)->Copy(*SourceUVLayer);
+										}
+										UVChannelParameterIndices[ParameterIndex] = (float)UVChannelIndex;
+									}
+								}
+							}
 
-									MaterialInstance->SetParentEditorOnly(Material);
+							// Find or create a material instance with the new UV channel indices
+							UMaterialInstanceConstant* MaterialInstance;
+							const FString MaterialPackagePath = FPackageName::GetLongPackagePath(Material->GetOutermost()->GetName());
+							const FString MaterialName = Material->GetName() + TEXT("_UVResized_");
+							const FString MaterialPackageName = FPaths::Combine(MaterialPackagePath, MaterialName);
+							FString MaterialPackageNameNumbered;
 
-									// Re-assign the UVChannels in this new material instance
-									for (const FString& UVChannelParameterName : UVChannelParameterNames)
+							for (int32 Suffix = 0; ; ++Suffix)
+							{
+								// Check if a material instance already exists and is identical
+								MaterialPackageNameNumbered = MaterialPackageName + FString::FromInt(Suffix);
+
+								if (const UPackage* const Package = FindPackage(nullptr, *MaterialPackageNameNumbered))
+								{
+									MaterialInstance = Cast<UMaterialInstanceConstant>(Package->FindAssetInPackage());
+								}
+								else
+								{
+									MaterialInstance = nullptr;
+									break;  // No existing object named with this suffix, a new material instance will need to be created
+								}
+								if (MaterialInstance && MaterialInstance->Parent == Material)
+								{
+									bool bIsIdentical = true;
+									for (int32 ParameterIndex = 0; ParameterIndex < UVChannelParameterNames.Num(); ++ParameterIndex)
 									{
+										const FString& UVChannelParameterName = UVChannelParameterNames[ParameterIndex];
 										float Value;
-										if (MaterialInstance->GetScalarParameterValue(*UVChannelParameterName, Value))
+										if (!MaterialInstance->GetScalarParameterValue(*UVChannelParameterName, Value) ||
+											Value != UVChannelParameterIndices[ParameterIndex])
 										{
-											const int32 SourceUVChannelIndex = (int32)Value;
-											if (UVChannelRemaps.IsValidIndex(SourceUVChannelIndex))
-											{
-												int32& UVChannelIndex = UVChannelRemaps[SourceUVChannelIndex];
-												if (UVChannelIndex == INDEX_NONE)
-												{
-													// Duplicate UV channel
-													UVChannelIndex = DynamicMesh.Attributes()->NumUVLayers();
-													DynamicMesh.Attributes()->SetNumUVLayers(UVChannelIndex + 1);
-													const UE::Geometry::FDynamicMeshUVOverlay* const SourceUVLayer = DynamicMesh.Attributes()->GetUVLayer(SourceUVChannelIndex);
-													check(SourceUVLayer);
-													DynamicMesh.Attributes()->GetUVLayer(UVChannelIndex)->Copy(*SourceUVLayer);
-												}
-												MaterialInstance->SetScalarParameterValueEditorOnly(*UVChannelParameterName, (float)UVChannelIndex);
-											}
+											bIsIdentical = false;  // Parameter don't match, check with another Suffix
+											break;
 										}
 									}
+									if (bIsIdentical)
+									{
+										break;  // Found a valid existing instance
+									}
+								}
+							}
+
+							if (!MaterialInstance)
+							{
+								MaterialInstance = Cast<UMaterialInstanceConstant>(Context.AddAsset(MaterialPackageNameNumbered, UMaterialInstanceConstant::StaticClass()));
+
+								if (ensure(MaterialInstance))
+								{
+									// Setup the new material instance
+									MaterialInstance->SetParentEditorOnly(Material);
+
+									for (int32 ParameterIndex = 0; ParameterIndex < UVChannelParameterNames.Num(); ++ParameterIndex)
+									{
+										const FString& UVChannelParameterName = UVChannelParameterNames[ParameterIndex];
+										MaterialInstance->SetScalarParameterValueEditorOnly(*UVChannelParameterName, UVChannelParameterIndices[ParameterIndex]);
+									}
 
 									// Finalize the new material asset creation
 									MaterialInstance->PreEditChange(nullptr);
 									MaterialInstance->PostEditChange();
-									MaterialInstance->MarkPackageDirty();
-									FAssetRegistryModule::AssetCreated(MaterialInstance);
-
-									OutMaterials.Emplace(MaterialInstance);
-									continue;
 								}
 							}
+
+							OutMaterials.Emplace(MaterialInstance);
+							continue;
 						}
 						OutMaterials.Emplace(Material);
 					}
+
 					// Create the output mesh
 					TObjectPtr<UDataflowMesh> OutMesh = NewObject<UDataflowMesh>();
 					OutMesh->SetDynamicMesh(MoveTemp(DynamicMesh));
diff --git a/Engine/Source/Runtime/Experimental/Dataflow/Core/Private/Dataflow/DataflowContextAssetStore.cpp b/Engine/Source/Runtime/Experimental/Dataflow/Core/Private/Dataflow/DataflowContextAssetStore.cpp
--- a/Engine/Source/Runtime/Experimental/Dataflow/Core/Private/Dataflow/DataflowContextAssetStore.cpp
+++ b/Engine/Source/Runtime/Experimental/Dataflow/Core/Private/Dataflow/DataflowContextAssetStore.cpp
@@ -2,8 +2,10 @@
 
 #include "Dataflow/DataflowContextAssetStore.h"
 
-#include "UObject/Package.h"
+#include "AssetRegistry/AssetRegistryModule.h"
 #include "Misc/Paths.h"
+#include "UObject/Package.h"
+#include "UObject/SavePackage.h"
 
 namespace UE::Dataflow
 {
@@ -58,13 +60,15 @@ namespace UE::Dataflow
 			{
 				UniqueAssetPath = MakeUniqueObjectName(nullptr, UPackage::StaticClass(), FName(UniqueAssetPath)).ToString();
 			}
-			// create the corresponding package
+			// create and save the corresponding package
 			if (UPackage* const NewAssetPackage = CreatePackage(*UniqueAssetPath))
 			{
 				const FString AssetName = FPaths::GetBaseFilename(UniqueAssetPath);
-				if (UObject* ObjectToCommit = DuplicateObject_Internal(InAssetData.TransientAsset->GetClass(), InAssetData.TransientAsset, NewAssetPackage, FName(AssetName)))
+				if (UObject* const ObjectToCommit = DuplicateObject_Internal(InAssetData.TransientAsset->GetClass(), InAssetData.TransientAsset, NewAssetPackage, FName(AssetName)))
 				{
+					ObjectToCommit->SetFlags(RF_Standalone);
 					ObjectToCommit->MarkPackageDirty();
+					FAssetRegistryModule::AssetCreated(ObjectToCommit);
 					return ObjectToCommit;
 				}
 			}
