Fix CIS Automation test error from a previous submit.

#rnx

diff --git a/Engine/Plugins/Experimental/MeshResizing/Source/MeshResizingCore/Public/MeshResizing/CustomRegionResizing.h b/Engine/Plugins/Experimental/MeshResizing/Source/MeshResizingCore/Public/MeshResizing/CustomRegionResizing.h
--- a/Engine/Plugins/Experimental/MeshResizing/Source/MeshResizingCore/Public/MeshResizing/CustomRegionResizing.h
+++ b/Engine/Plugins/Experimental/MeshResizing/Source/MeshResizingCore/Public/MeshResizing/CustomRegionResizing.h
@@ -38,24 +38,24 @@ struct FMeshResizingCustomRegion
 	TArray<FVector3f> RegionVertexCoords; // Trilinear interpolation within Bounds: (0,0,0) = RegionBoundsCentroid - RegionBoundsExtents; (1,1,1) = RegionBoundsCentroid + RegionBoundsExtents
 
 	UPROPERTY()
-	int32 SourceFaceIndex;
+	int32 SourceFaceIndex = INDEX_NONE;
 	UPROPERTY()
-	FVector3f SourceBaryCoords;
+	FVector3f SourceBaryCoords = FVector3f(0.f);
 	UPROPERTY()
-	FVector3d SourceOrigin;
+	FVector3d SourceOrigin = FVector3d(0.f);
 	UPROPERTY()
-	FVector3f SourceAxis0;
+	FVector3f SourceAxis0 = FVector3f(0.f);
 	UPROPERTY()
-	FVector3f SourceAxis1;
+	FVector3f SourceAxis1 = FVector3f(0.f);
 	UPROPERTY()
-	FVector3f SourceAxis2;
+	FVector3f SourceAxis2 = FVector3f(0.f);
 
 	// Relative to Source
 	UPROPERTY()
-	FVector3f RegionBoundsCentroid;
+	FVector3f RegionBoundsCentroid = FVector3f(0.f);
 
 	UPROPERTY()
-	FVector3f RegionBoundsExtents;
+	FVector3f RegionBoundsExtents = FVector3f(0.f);
 
 	void Reset()
 	{
