MeshResizing: add mesh wrap option for building resizeable proxy

diff --git a/Engine/Plugins/Experimental/MeshResizing/MeshResizing.uplugin b/Engine/Plugins/Experimental/MeshResizing/MeshResizing.uplugin
--- a/Engine/Plugins/Experimental/MeshResizing/MeshResizing.uplugin
+++ b/Engine/Plugins/Experimental/MeshResizing/MeshResizing.uplugin
@@ -24,6 +24,11 @@
 			"Name": "MeshResizingDataflowNodes",
 			"Type": "Runtime",
 			"LoadingPhase": "Default"
+		},
+		{
+			"Name": "MeshResizingEditorTools",
+			"Type": "Editor",
+			"LoadingPhase": "Default"
 		}
 	],
 	"Plugins": [
@@ -35,9 +40,13 @@
 			"Name": "GeometryCollectionPlugin",
 			"Enabled": true
 		},
+		{
+			"Name": "GeometryProcessing",
+			"Enabled": true
+		},
 		{
 			"Name": "HairStrands",
-			"Enabled":  true
+			"Enabled": true
 		}
 	]
 }
diff --git a/Engine/Plugins/Experimental/MeshResizing/Source/MeshResizingEditorTools/MeshResizingEditorTools.Build.cs b/Engine/Plugins/Experimental/MeshResizing/Source/MeshResizingEditorTools/MeshResizingEditorTools.Build.cs
new file mode 100644
--- /dev/null
+++ b/Engine/Plugins/Experimental/MeshResizing/Source/MeshResizingEditorTools/MeshResizingEditorTools.Build.cs
@@ -0,0 +1,42 @@
+// Copyright Epic Games, Inc. All Rights Reserved.
+
+using UnrealBuildTool;
+public class MeshResizingEditorTools : ModuleRules
+{
+	public MeshResizingEditorTools(ReadOnlyTargetRules Target) : base(Target)
+	{
+		PublicIncludePathModuleNames.AddRange(
+			new string[]
+			{
+			}
+		);
+
+		PublicDependencyModuleNames.AddRange(
+			new string[]
+			{
+				"Core",
+				"CoreUObject",
+				"Engine",
+				"GeometryCore",
+				"UnrealEd",
+			}
+		);
+		PrivateDependencyModuleNames.AddRange(
+			new string[]
+			{
+				"DataflowCore",
+				"DataflowEngine",
+				"DataflowEditor",
+				"DynamicMesh",
+				"GeometryFramework",
+				"InteractiveToolsFramework",
+				"MeshModelingTools",
+				"MeshResizingCore",
+				"MeshResizingDataflowNodes",
+				"ModelingComponents",
+				"Slate",
+				"SlateCore",
+			}
+			);
+	}
+}
diff --git a/Engine/Plugins/Experimental/MeshResizing/Source/MeshResizingEditorTools/Private/MeshResizing/MeshResizingEditorToolsModule.cpp b/Engine/Plugins/Experimental/MeshResizing/Source/MeshResizingEditorTools/Private/MeshResizing/MeshResizingEditorToolsModule.cpp
new file mode 100644
--- /dev/null
+++ b/Engine/Plugins/Experimental/MeshResizing/Source/MeshResizingEditorTools/Private/MeshResizing/MeshResizingEditorToolsModule.cpp
@@ -0,0 +1,40 @@
+// Copyright Epic Games, Inc. All Rights Reserved.
+
+#include "Modules/ModuleManager.h"
+
+#include "Dataflow/DataflowToolRegistry.h"
+#include "MeshResizing/MeshResizingToolActionCommandBindings.h"
+#include "MeshResizing/MeshWrapNode.h"
+
+#define LOCTEXT_NAMESPACE "MeshResizingEditorToolsModule"
+
+namespace UE::MeshResizing
+{
+	class FMeshResizingEditorTools : public IModuleInterface
+	{
+	public:
+		virtual void StartupModule() override
+		{
+			UE::Dataflow::FDataflowToolRegistry& ToolRegistry = UE::Dataflow::FDataflowToolRegistry::Get();
+		
+			const TSharedRef<const FMeshResizingToolActionCommandBindings> MeshResizingToolActions = MakeShared<FMeshResizingToolActionCommandBindings>();
+
+			ToolRegistry.AddNodeToToolMapping(FMeshWrapLandmarksNode::StaticType(),          // Node type
+				NewObject<UMeshWrapLandmarkSelectionToolBuilder>(),                          // Tool builder
+				MeshResizingToolActions,                                                     // Tool action bindings
+				FSlateIcon(),                                                                // Button icon 
+				LOCTEXT("AddMeshWrapLandmarksNodeButtonText", "Mesh Resizing Landmarks"),    // Button text
+				FName("General"));                                                           // Tool category
+		}
+
+		virtual void ShutdownModule() override
+		{
+			UE::Dataflow::FDataflowToolRegistry& ToolRegistry = UE::Dataflow::FDataflowToolRegistry::Get();
+			ToolRegistry.RemoveNodeToToolMapping(FMeshWrapLandmarksNode::StaticType());
+		}
+	};
+}
+
+IMPLEMENT_MODULE(UE::MeshResizing::FMeshResizingEditorTools, MeshResizingEditorTools)
+
+#undef LOCTEXT_NAMESPACE
diff --git a/Engine/Plugins/Experimental/MeshResizing/Source/MeshResizingEditorTools/Private/MeshResizing/MeshResizingToolActionCommandBindings.cpp b/Engine/Plugins/Experimental/MeshResizing/Source/MeshResizingEditorTools/Private/MeshResizing/MeshResizingToolActionCommandBindings.cpp
new file mode 100644
--- /dev/null
+++ b/Engine/Plugins/Experimental/MeshResizing/Source/MeshResizingEditorTools/Private/MeshResizing/MeshResizingToolActionCommandBindings.cpp
@@ -0,0 +1,36 @@
+// Copyright Epic Games, Inc. All Rights Reserved.
+
+#include "MeshResizing/MeshResizingToolActionCommandBindings.h"
+
+#define LOCTEXT_NAMESPACE "MeshResizingToolActionCommandBindings"
+
+FMeshWrapLandmarkSelectionToolActionCommands::FMeshWrapLandmarkSelectionToolActionCommands() : FMeshResizingToolActionCommands<FMeshWrapLandmarkSelectionToolActionCommands, UMeshWrapLandmarkSelectionTool>(
+	TEXT("MeshWrapLandmarkSelectionToolContext"), LOCTEXT("MeshWrapLandmarkSelectionToolContext", "Mesh Wrap Landmarks Selection Tool Context"))
+{}
+
+FMeshResizingToolActionCommandBindings::FMeshResizingToolActionCommandBindings()
+{
+	// Note: if a TCommands<> doesn't actually register any commands then it will be deleted. Only WeightMapPaintTool currently has key commands, but we will include the other tools
+	// here so that hotkeys can be added to them in the future. This means we need to check if the objects are registered before trying to use them below.
+
+	FMeshWrapLandmarkSelectionToolActionCommands::Register();
+}
+
+void FMeshResizingToolActionCommandBindings::UnbindActiveCommands(const TSharedPtr<FUICommandList>& UICommandList) const
+{
+	if (FMeshWrapLandmarkSelectionToolActionCommands::IsRegistered())
+	{
+		FMeshWrapLandmarkSelectionToolActionCommands::Get().UnbindActiveCommands(UICommandList);
+	}
+}
+
+void FMeshResizingToolActionCommandBindings::BindCommandsForCurrentTool(const TSharedPtr<FUICommandList>& UICommandList, UInteractiveTool* Tool) const
+{
+	if (ExactCast<UMeshWrapLandmarkSelectionTool>(Tool) && FMeshWrapLandmarkSelectionToolActionCommands::IsRegistered())
+	{
+		FMeshWrapLandmarkSelectionToolActionCommands::Get().BindCommandsForCurrentTool(UICommandList, Tool);
+	}
+}
+
+#undef LOCTEXT_NAMESPACE
+
diff --git a/Engine/Plugins/Experimental/MeshResizing/Source/MeshResizingEditorTools/Private/MeshResizing/MeshResizingToolActionCommandBindings.h b/Engine/Plugins/Experimental/MeshResizing/Source/MeshResizingEditorTools/Private/MeshResizing/MeshResizingToolActionCommandBindings.h
new file mode 100644
--- /dev/null
+++ b/Engine/Plugins/Experimental/MeshResizing/Source/MeshResizingEditorTools/Private/MeshResizing/MeshResizingToolActionCommandBindings.h
@@ -0,0 +1,50 @@
+// Copyright Epic Games, Inc. All Rights Reserved.
+
+#pragma once
+
+#include "Dataflow/DataflowToolRegistry.h"
+#include "Tools/InteractiveToolsCommands.h"
+#include "Styling/AppStyle.h"			// Style name is a required parameter of TInteractiveToolCommands/FBindingContext
+
+// For GetToolDefaultObjectList in class template 
+#include "MeshResizing/MeshWrapLandmarkSelectionTool.h"	
+
+
+// TInteractiveToolCommands<> are typically used to bind the current set of available hotkey commands when a tool starts/ends. However we cannot store multiple actions 
+// with the same key activation in a single TInteractiveToolCommands object (even if they are only active in different tools), so we end up creating one TInteractiveToolCommands per tool. 
+// This is also what Modeling Tools does -- see comments in ModelingToolsActions.h.
+
+
+// Base class with TInteractiveToolCommands<> boilerplate
+
+template<typename T, typename ToolClass>
+class FMeshResizingToolActionCommands : public TInteractiveToolCommands<T>
+{
+public:
+	FMeshResizingToolActionCommands(const FName InContextName, const FText& InContextDesc) : TInteractiveToolCommands<T>(InContextName, InContextDesc, NAME_None, FAppStyle::GetAppStyleSetName())
+	{}
+
+	virtual void GetToolDefaultObjectList(TArray<UInteractiveTool*>& ToolCDOs) override
+	{
+		ToolCDOs.Add(GetMutableDefault<ToolClass>());
+	}
+};
+
+// Derived classes for individual tools
+
+class FMeshWrapLandmarkSelectionToolActionCommands : public FMeshResizingToolActionCommands<FMeshWrapLandmarkSelectionToolActionCommands, UMeshWrapLandmarkSelectionTool>
+{
+public:
+	FMeshWrapLandmarkSelectionToolActionCommands();
+};
+
+// Tool action registry entry
+
+class FMeshResizingToolActionCommandBindings : public UE::Dataflow::FDataflowToolRegistry::IDataflowToolActionCommands
+{
+
+public:
+	FMeshResizingToolActionCommandBindings();
+	virtual void UnbindActiveCommands(const TSharedPtr<FUICommandList>& UICommandList) const override;
+	virtual void BindCommandsForCurrentTool(const TSharedPtr<FUICommandList>& UICommandList, UInteractiveTool* Tool) const override;
+};
diff --git a/Engine/Plugins/Experimental/MeshResizing/Source/MeshResizingEditorTools/Private/MeshResizing/MeshWrapLandmarkSelectionTool.cpp b/Engine/Plugins/Experimental/MeshResizing/Source/MeshResizingEditorTools/Private/MeshResizing/MeshWrapLandmarkSelectionTool.cpp
new file mode 100644
--- /dev/null
+++ b/Engine/Plugins/Experimental/MeshResizing/Source/MeshResizingEditorTools/Private/MeshResizing/MeshWrapLandmarkSelectionTool.cpp
@@ -0,0 +1,498 @@
+// Copyright Epic Games, Inc. All Rights Reserved.
+
+#include "MeshResizing/MeshWrapLandmarkSelectionTool.h"
+#include "Dataflow/DataflowContextObject.h"
+#include "Dataflow/DataflowObject.h"
+#include "Engine/Canvas.h"
+#include "Engine/Engine.h"
+#include "Materials/Material.h"
+#include "Materials/MaterialInterface.h"
+#include "TargetInterfaces/PrimitiveComponentBackedTarget.h"
+#include "Util/ColorConstants.h"
+#include "ContextObjectStore.h"
+#include "GroupTopology.h"
+#include "InteractiveToolManager.h"
+#include "ModelingToolTargetUtil.h"
+#include "PreviewMesh.h"
+#include "SceneView.h"
+#include "ToolContextInterfaces.h"
+#include "ToolSetupUtil.h"
+#include "ToolTargetManager.h"
+
+#include UE_INLINE_GENERATED_CPP_BY_NAME(MeshWrapLandmarkSelectionTool)
+
+#define LOCTEXT_NAMESPACE "MeshWrapLandmarkSelectionTool"
+
+namespace UE::MeshResizing::Private
+{
+	static FLinearColor PseudoRandomColor(int32 NumColorRotations)
+	{
+		constexpr uint8 Spread = 157;  // Prime number that gives a good spread of colors without getting too similar as a rand might do.
+		uint8 Seed = Spread;
+		NumColorRotations = FMath::Abs(NumColorRotations);
+		for (int32 Rotation = 0; Rotation < NumColorRotations; ++Rotation)
+		{
+			Seed += Spread;
+		}
+		return FLinearColor::MakeFromHSV8(Seed, 180, 140);
+	}
+
+	static void DrawText(FCanvas* Canvas, const FSceneView* SceneView, const FVector& Pos, const FString& Text, const FLinearColor& Color, const float Scale = 1.f)
+	{
+		if (Canvas && SceneView)
+		{
+			FVector2D PixelLocation;
+			if (SceneView->WorldToPixel(Pos, PixelLocation))
+			{
+				// WorldToPixel doesn't account for DPIScale
+				const float DPIScale = Canvas->GetDPIScale();
+				FCanvasTextItem TextItem(PixelLocation / DPIScale, FText::FromString(Text), GEngine->GetSmallFont(), Color);
+				TextItem.Scale = FVector2D::UnitVector * Scale;
+				TextItem.EnableShadow(FLinearColor::Black);
+				TextItem.Draw(Canvas);
+			}
+		}
+	}
+}
+
+// ------------------- Selection Mechanic -------------------
+
+bool UMeshWrapLandmarkSelectionMechanic::UpdateSelection(const FRay& WorldRay, FVector3d& LocalHitPositionOut, FVector3d& LocalHitNormalOut)
+{
+	using namespace UE::Geometry;
+	if (bCtrlToggle)
+	{
+		// only allow selecting existing landmarks with Ctrl
+		if (const UMeshWrapLandmarkSelectionTool* const SelectionTool = Cast<UMeshWrapLandmarkSelectionTool>(ParentTool))
+		{
+			FRay3d LocalRay(TargetTransform.InverseTransformPosition((FVector3d)WorldRay.Origin),
+				TargetTransform.InverseTransformVector((FVector3d)WorldRay.Direction));
+			UE::Geometry::Normalize(LocalRay.Direction);
+
+			const FGroupTopologySelection PreviousSelection = PersistentSelection;
+
+			FVector3d LocalPosition, LocalNormal;
+			FGroupTopologySelection Selection;
+			FGroupTopologySelector::FSelectionSettings TopoSelectorSettings = GetTopoSelectorSettings(CameraState.bIsOrthographic);
+			if (TopoSelector->FindSelectedElement(TopoSelectorSettings, LocalRay, Selection, LocalPosition, LocalNormal) && Selection.SelectedCornerIDs.Num() == 1)
+			{
+				const int32 SelectedCornerID = Selection.GetASelectedCornerID();
+				const int32 LandmarkIndex = SelectionTool->GetFirstLandmarkWithID(SelectedCornerID);
+				if (LandmarkIndex != INDEX_NONE)
+				{
+					LocalHitPositionOut = LocalPosition;
+					LocalHitNormalOut = LocalNormal;
+
+					if (PersistentSelection != Selection || bHadCtrlOnSelection != bCtrlToggle)
+					{
+						bHadCtrlOnSelection = bCtrlToggle;
+						PersistentSelection = Selection;
+						SelectionTimestamp++;
+						OnSelectionChanged.Broadcast();
+						return true;
+					}
+				}
+			}
+		}
+		return false;
+	}
+	bHadCtrlOnSelection = false;
+	bHadShiftOnSelection = bShiftToggle;
+	return UPolygonSelectionMechanic::UpdateSelection(WorldRay, LocalHitPositionOut, LocalHitNormalOut);
+}
+
+
+// ------------------- Tool -------------------
+	
+
+void UMeshWrapLandmarkSelectionTool::InitializeSculptMeshFromTarget()
+{
+	using namespace UE::Geometry;
+
+	//
+	// Preview 
+	//
+
+	// (Re-)Create the preview mesh
+	if (PreviewMesh)
+	{
+		PreviewMesh->Disconnect();
+	}
+	PreviewMesh = NewObject<UPreviewMesh>(this, TEXT("PreviewMesh"));
+	PreviewMesh->CreateInWorld(GetTargetWorld(), FTransform::Identity);
+
+	ToolSetupUtil::ApplyRenderingConfigurationToPreview(PreviewMesh, Target);
+
+	// We will use the preview mesh's spatial data structure
+	PreviewMesh->bBuildSpatialDataStructure = true;
+
+	// set materials
+	FComponentMaterialSet MaterialSet = UE::ToolTarget::GetMaterialSet(Target);
+	PreviewMesh->SetMaterials(MaterialSet.Materials);
+
+	// configure secondary render material for selected triangles
+	// NOTE: the material returned by ToolSetupUtil::GetSelectionMaterial has a checkerboard pattern on back faces which makes it hard to use
+	UMaterialInterface* Material = LoadObject<UMaterial>(nullptr, TEXT("/MeshModelingToolsetExp/Materials/SculptMaterial"));
+	if (Material != nullptr)
+	{
+		if (UMaterialInstanceDynamic* MatInstance = UMaterialInstanceDynamic::Create(Material, GetToolManager()))
+		{
+			MatInstance->SetVectorParameterValue(TEXT("Color"), FLinearColor::Yellow);
+			PreviewMesh->SetSecondaryRenderMaterial(MatInstance);
+		}
+	}
+
+	PreviewMesh->SetTangentsMode(EDynamicMeshComponentTangentsMode::AutoCalculated);
+	PreviewMesh->UpdatePreview(UE::ToolTarget::GetDynamicMeshCopy(Target));
+	PreviewMesh->SetVisible(true);
+
+	// Hide input target mesh
+	UE::ToolTarget::HideSourceObject(Target);
+}
+
+
+
+void UMeshWrapLandmarkSelectionTool::Setup()
+{
+	using namespace UE::Geometry;
+
+	InitializeSculptMeshFromTarget();
+
+	//
+	// SelectionMechanic
+	//
+
+	SelectionMechanic = NewObject<UMeshWrapLandmarkSelectionMechanic>(this);
+	SelectionMechanic->bAddSelectionFilterPropertiesToParentTool = false;   // We'll do this ourselves later
+	SelectionMechanic->Setup(this);
+	SelectionMechanic->Properties->RestoreProperties(this);
+	SelectionMechanic->Properties->bDisplayPolygroupReliantControls = false;	// this is for polygroup-specific selections like edge loops
+	
+	SelectionMechanic->Properties->bCanSelectVertices = true;
+	SelectionMechanic->Properties->bCanSelectEdges = false;
+	SelectionMechanic->Properties->bCanSelectFaces = false;
+	SelectionMechanic->Properties->bSelectVertices = true;
+	SelectionMechanic->Properties->bSelectEdges = false;
+	SelectionMechanic->Properties->bSelectFaces = false;
+
+	SelectionMechanic->Properties->bEnableMarquee = false;
+	
+	SelectionMechanic->SetShowEdges(false);
+	SelectionMechanic->SetShowSelectableCorners(false);
+
+	SelectionMechanic->PolyEdgesRenderer.DepthBias = 0.01f;
+	SelectionMechanic->PolyEdgesRenderer.LineThickness = 1.0f;
+	SelectionMechanic->PolyEdgesRenderer.PointSize = 2.0f;
+
+	SelectionMechanic->SelectionRenderer.DepthBias = 0.01f;
+	SelectionMechanic->SelectionRenderer.LineThickness = 1.0f;
+	SelectionMechanic->SelectionRenderer.PointSize = 2.0f;
+
+	SelectionMechanic->OnSelectionChanged.AddWeakLambda(this, [this]()
+	{
+		UpdateLandmarkFromSelection();
+	});
+
+	// Disallow adding/removing to the selection with shift/ctrl
+	SelectionMechanic->SetShouldAddToSelectionFunc([]() { return false; });
+	SelectionMechanic->SetShouldRemoveFromSelectionFunc([]() { return false; });
+
+	// Set up Topology and SelectionMechanic using Preview's DynamicMesh
+	PreviewMesh->ProcessMesh([this](const FDynamicMesh3& Mesh)
+	{
+		Topology = MakeUnique<FTriangleGroupTopology>(&Mesh, true);
+
+		SelectionMechanic->Initialize(&Mesh,
+			FTransform3d(),
+			GetTargetWorld(),
+			Topology.Get(),
+			[this]() { return PreviewMesh->GetSpatial(); });
+	});
+
+
+	//
+	// Properties
+	//
+
+	ToolProperties = NewObject<UMeshWrapLandmarkSelectionToolProperties>();
+	ToolProperties->WatchProperty(ToolProperties->bShowVertices, [this](bool bNewShowVertices)
+	{
+		SelectionMechanic->SetShowSelectableCorners(bNewShowVertices);
+	});
+
+	ToolProperties->WatchProperty(ToolProperties->bShowEdges, [this](bool bNewShowEdges)
+	{
+		SelectionMechanic->SetShowEdges(bNewShowEdges);
+	});
+
+	// Order of operations is important here:
+	// ToolProperties->WatchProperty(ToolProperties->Name) should happen after GetSelectedNodeInfo so that we can set Landmark data.
+
+	ToolProperties->RestoreProperties(this);
+	 
+	// Initialize the Selection from the selected Dataflow node
+	SetPropertiesFromSelectedNode();
+	SelectionNodeToUpdate->bCanDebugDraw = false; // Disable node's debug draw to not double draw
+	SelectionNodeToUpdate->Invalidate();
+	ToolProperties->WatchProperty(ToolProperties->Landmarks, [this](const TArray<FMeshWrapToolLandmark>& NewLandmarks)
+		{
+			bAnyChangeMade = true;
+		});
+	ToolProperties->WatchProperty(ToolProperties->CurrentEditableLandmark, [this](int32 NewCurrentEditableLandmark)
+		{
+			SetCurrentEditableLandmark(NewCurrentEditableLandmark);
+		});
+
+	AddToolPropertySource(ToolProperties);
+}
+
+void UMeshWrapLandmarkSelectionTool::SetPropertiesFromSelectedNode()
+{
+	SelectionNodeToUpdate = DataflowContextObject->GetSelectedNodeOfType<FMeshWrapLandmarksNode>();
+	check(SelectionNodeToUpdate);
+	ToolProperties->Landmarks.SetNum(SelectionNodeToUpdate->Landmarks.Num());
+	for (int32 Index = 0; Index < SelectionNodeToUpdate->Landmarks.Num(); ++Index)
+	{
+		ToolProperties->Landmarks[Index].Identifier = SelectionNodeToUpdate->Landmarks[Index].Identifier;
+		ToolProperties->Landmarks[Index].VertexIndex = SelectionNodeToUpdate->Landmarks[Index].VertexIndex;
+
+	}
+	ToolProperties->CurrentEditableLandmark = INDEX_NONE;
+}
+
+int32 UMeshWrapLandmarkSelectionTool::GetFirstLandmarkWithID(const int32 Id) const
+{
+	if (Id == INDEX_NONE)
+	{
+		// Special case--don't want to select a landmark when no vertex is selected.
+		return INDEX_NONE;
+	}
+	for (int32 Index = 0; Index < ToolProperties->Landmarks.Num(); ++Index)
+	{
+		if (ToolProperties->Landmarks[Index].VertexIndex == Id)
+		{
+			return Index;
+		}
+	}
+	return INDEX_NONE;
+}
+
+void UMeshWrapLandmarkSelectionTool::SetCurrentEditableLandmark(int32 Index)
+{
+	if (ToolProperties->CurrentEditableLandmark != Index)
+	{
+		ToolProperties->CurrentEditableLandmark = ToolProperties->Landmarks.IsValidIndex(Index) ? Index : INDEX_NONE;
+		if (ToolProperties->CurrentEditableLandmark != INDEX_NONE)
+		{
+			// Update current selection
+			FGroupTopologySelection LandmarkSelection;
+			PreviewMesh->ProcessMesh([this, &LandmarkSelection](const UE::Geometry::FDynamicMesh3& Mesh)
+				{
+					if (Mesh.IsVertex(ToolProperties->Landmarks[ToolProperties->CurrentEditableLandmark].VertexIndex))
+					{
+						LandmarkSelection.SelectedCornerIDs = TSet<int32>({ ToolProperties->Landmarks[ToolProperties->CurrentEditableLandmark].VertexIndex });
+					}
+				});
+			constexpr bool bBroadcastChange = false;
+			SelectionMechanic->SetSelection(LandmarkSelection, bBroadcastChange);
+		}
+	}
+}
+
+void UMeshWrapLandmarkSelectionTool::UpdateLandmarkFromSelection()
+{
+	const UE::Geometry::FGroupTopologySelection& Selection = SelectionMechanic->GetActiveSelection();
+	const int32 SelectedCornerID = Selection.GetASelectedCornerID();
+	if (SelectionMechanic->GetHadCtrlOnSelection())
+	{
+		const int32 LandmarkIndex = GetFirstLandmarkWithID(SelectedCornerID);
+		if (ToolProperties->Landmarks.IsValidIndex(LandmarkIndex))
+		{
+			SetCurrentEditableLandmark(LandmarkIndex);
+		}
+	}
+	else
+	{
+		bool bIsValidVertex = false;
+		if (SelectedCornerID != INDEX_NONE)
+		{
+			PreviewMesh->ProcessMesh([this, SelectedCornerID, &bIsValidVertex](const UE::Geometry::FDynamicMesh3& Mesh)
+				{
+					bIsValidVertex = Mesh.IsVertex(SelectedCornerID);
+				});
+		}
+		
+		if (SelectionMechanic->GetHadShiftOnSelection() || !ToolProperties->Landmarks.IsValidIndex(ToolProperties->CurrentEditableLandmark))
+		{
+			if (bIsValidVertex)
+			{
+				// Add a new landmark
+				bAnyChangeMade = true;
+				FMeshWrapToolLandmark& NewLandmark = ToolProperties->Landmarks.AddDefaulted_GetRef();
+				NewLandmark.Identifier = TEXT("NewLandmark_") + FString::FromInt(ToolProperties->Landmarks.Num()-1);
+				NewLandmark.VertexIndex = SelectedCornerID;
+				ToolProperties->CurrentEditableLandmark = ToolProperties->Landmarks.Num() - 1;
+			}
+		}
+		else
+		{
+			bAnyChangeMade = true;
+			ToolProperties->Landmarks[ToolProperties->CurrentEditableLandmark].VertexIndex = bIsValidVertex ? SelectedCornerID : INDEX_NONE;
+		}
+	}
+}
+
+void UMeshWrapLandmarkSelectionTool::OnShutdown(EToolShutdownType ShutdownType)
+{
+	if (ShutdownType == EToolShutdownType::Accept && CanAccept())
+	{
+		GetToolManager()->BeginUndoTransaction(LOCTEXT("MeshWrapLandmarkSelectionToolTransactionName", "Mesh Wrap Landmark"));
+		UpdateSelectedNode();
+		GetToolManager()->EndUndoTransaction();
+
+		SelectionNodeToUpdate->Invalidate();
+	}
+
+	SelectionMechanic->Properties->SaveProperties(this);
+	ToolProperties->SaveProperties(this);
+
+	if (PreviewMesh != nullptr)
+	{
+		UE::ToolTarget::ShowSourceObject(Target);
+		PreviewMesh->Disconnect();
+		PreviewMesh = nullptr;
+	}
+
+	SelectionMechanic->Shutdown();
+
+	Topology.Reset();
+	SelectionNodeToUpdate->bCanDebugDraw = true;
+}
+
+void UMeshWrapLandmarkSelectionTool::Render(IToolsContextRenderAPI* RenderAPI)
+{
+	check(PreviewMesh);
+	check(RenderAPI);
+
+	SelectionMechanic->Render(RenderAPI);
+
+	// Display landmark points
+	FPrimitiveDrawInterface* const PDI = RenderAPI->GetPrimitiveDrawInterface();
+	PreviewMesh->ProcessMesh([this, PDI](const UE::Geometry::FDynamicMesh3& Mesh)
+		{
+			for (int32 Index = 0; Index < ToolProperties->Landmarks.Num(); ++Index)
+			{
+				if (Mesh.IsVertex(ToolProperties->Landmarks[Index].VertexIndex))
+				{
+					if (Index != ToolProperties->CurrentEditableLandmark)
+					{
+						PDI->DrawPoint(Mesh.GetVertexRef(ToolProperties->Landmarks[Index].VertexIndex), UE::MeshResizing::Private::PseudoRandomColor(Index), 10.f, SDPG_World);
+					}
+				}
+			}
+		});
+}
+
+void UMeshWrapLandmarkSelectionTool::DrawHUD(FCanvas* Canvas, IToolsContextRenderAPI* RenderAPI)
+{
+	check(RenderAPI);
+	check(Canvas);
+
+	SelectionMechanic->DrawHUD(Canvas, RenderAPI);
+
+	static const FLinearColor SelectionColor = UE::Geometry::LinearColors::Gold3f();
+	// Display landmark names
+	if (const FSceneView* const SceneView = RenderAPI->GetSceneView())
+	{
+		PreviewMesh->ProcessMesh([this, Canvas, SceneView](const UE::Geometry::FDynamicMesh3& Mesh)
+			{
+				for (int32 Index = 0; Index < ToolProperties->Landmarks.Num(); ++Index)
+				{
+					if (Mesh.IsVertex(ToolProperties->Landmarks[Index].VertexIndex))
+					{
+						UE::MeshResizing::Private::DrawText(Canvas, SceneView, Mesh.GetVertexRef(ToolProperties->Landmarks[Index].VertexIndex), FString::FromInt(Index) + FString(" ") + ToolProperties->Landmarks[Index].Identifier,
+							Index == ToolProperties->CurrentEditableLandmark ? SelectionColor : UE::MeshResizing::Private::PseudoRandomColor(Index));
+					}
+				}
+			});
+	}
+}
+
+bool UMeshWrapLandmarkSelectionTool::CanAccept() const
+{
+	return bAnyChangeMade;
+}
+
+FBox UMeshWrapLandmarkSelectionTool::GetWorldSpaceFocusBox()
+{
+	static constexpr bool bWorld = true;
+	return FBox(SelectionMechanic->GetSelectionBounds(bWorld));
+}
+
+void UMeshWrapLandmarkSelectionTool::SetDataflowContextObject(TObjectPtr<UDataflowContextObject> InDataflowContextObject)
+{
+	DataflowContextObject = InDataflowContextObject;
+}
+
+void UMeshWrapLandmarkSelectionTool::UpdateSelectedNode()
+{
+	checkf(SelectionNodeToUpdate, TEXT("Expected non-null pointer to Selection Node"));
+
+	// Save previous state for undo
+	if (UDataflow* const Dataflow = DataflowContextObject->GetDataflowAsset())
+	{
+		GetToolManager()->GetContextTransactionsAPI()->AppendChange(Dataflow, 
+			FMeshWrapLandmarksNode::MakeSelectedNodeChange(*SelectionNodeToUpdate),
+			LOCTEXT("MeshWrapLandmarkUpdate", "Update Mesh Wrap Landmarks Node"));
+	}
+
+	check(SelectionNodeToUpdate);
+	SelectionNodeToUpdate->Landmarks.SetNum(ToolProperties->Landmarks.Num());
+	for (int32 Index = 0; Index < ToolProperties->Landmarks.Num(); ++Index)
+	{
+		SelectionNodeToUpdate->Landmarks[Index].Identifier = ToolProperties->Landmarks[Index].Identifier;
+		SelectionNodeToUpdate->Landmarks[Index].VertexIndex = ToolProperties->Landmarks[Index].VertexIndex;
+	}
+	
+}
+
+// ------------------- Tool Builder -------------------
+
+void UMeshWrapLandmarkSelectionToolBuilder::GetSupportedConstructionViewModes(const UDataflowContextObject& ContextObject, TArray<const UE::Dataflow::IDataflowConstructionViewMode*>& Modes) const
+{
+}
+
+bool UMeshWrapLandmarkSelectionToolBuilder::CanBuildTool(const FToolBuilderState& SceneState) const
+{
+	if (UDataflowContextObject* const DataflowContextObject = SceneState.ToolManager->GetContextObjectStore()->FindContext<UDataflowContextObject>())
+	{
+		return DataflowContextObject->GetSelectedNodeOfType<FMeshWrapLandmarksNode>() != nullptr && SceneState.TargetManager != nullptr &&(SceneState.TargetManager->CountSelectedAndTargetable(SceneState, GetTargetRequirements()) == 1);
+	}
+
+	return false;
+}
+
+const FToolTargetTypeRequirements& UMeshWrapLandmarkSelectionToolBuilder::GetTargetRequirements() const
+{
+	static FToolTargetTypeRequirements TypeRequirements(UPrimitiveComponentBackedTarget::StaticClass());
+	return TypeRequirements;
+}
+
+
+
+UInteractiveTool* UMeshWrapLandmarkSelectionToolBuilder::BuildTool(const FToolBuilderState& SceneState) const
+{
+	UMeshWrapLandmarkSelectionTool* const NewTool = NewObject<UMeshWrapLandmarkSelectionTool>(SceneState.ToolManager);
+
+	UToolTarget* const Target = SceneState.TargetManager->BuildFirstSelectedTargetable(SceneState, GetTargetRequirements());
+	NewTool->SetTarget(Target);
+	NewTool->SetWorld(SceneState.World);
+
+	if (UDataflowContextObject* const DataflowContextObject = SceneState.ToolManager->GetContextObjectStore()->FindContext<UDataflowContextObject>())
+	{
+		NewTool->SetDataflowContextObject(DataflowContextObject);
+	}
+
+	return NewTool;
+}
+#undef LOCTEXT_NAMESPACE
diff --git a/Engine/Plugins/Experimental/MeshResizing/Source/MeshResizingEditorTools/Private/MeshResizing/MeshWrapLandmarkSelectionTool.h b/Engine/Plugins/Experimental/MeshResizing/Source/MeshResizingEditorTools/Private/MeshResizing/MeshWrapLandmarkSelectionTool.h
new file mode 100644
--- /dev/null
+++ b/Engine/Plugins/Experimental/MeshResizing/Source/MeshResizingEditorTools/Private/MeshResizing/MeshWrapLandmarkSelectionTool.h
@@ -0,0 +1,153 @@
+// Copyright Epic Games, Inc. All Rights Reserved.
+
+#pragma once
+
+#include "MeshResizing/MeshWrapNode.h"
+#include "BaseTools/SingleSelectionMeshEditingTool.h"
+#include "DataflowEditorTools/DataflowEditorToolBuilder.h"
+#include "Selection/PolygonSelectionMechanic.h"
+#include "MeshWrapLandmarkSelectionTool.generated.h"
+
+#define UE_API MESHRESIZINGEDITORTOOLS_API
+
+class UDataflowContextObject;
+class UPreviewMesh;
+
+namespace UE::Geometry
+{
+	class FGroupTopology;
+	struct FGroupTopologySelection;
+}
+
+UCLASS()
+class UMeshWrapLandmarkSelectionMechanic : public UPolygonSelectionMechanic
+{
+	GENERATED_BODY()
+public:
+	virtual bool UpdateSelection(const FRay& WorldRay, FVector3d& LocalHitPositionOut, FVector3d& LocalHitNormalOut) override;
+
+	bool GetHadShiftOnSelection() const { return bHadShiftOnSelection; }
+	bool GetHadCtrlOnSelection() const { return bHadCtrlOnSelection; }
+
+private:
+	bool bHadShiftOnSelection = false;
+	bool bHadCtrlOnSelection = false;
+};
+
+/**
+ * Mesh Wrap Landmark Selection Tool Landmark
+ */
+USTRUCT()
+struct FMeshWrapToolLandmark
+{
+	GENERATED_USTRUCT_BODY()
+
+	// Landmarks will be matched by comparing Identifiers
+	UPROPERTY(EditAnywhere, Category = "Mesh Wrap")
+	FString Identifier = TEXT("");
+
+	UPROPERTY(VisibleAnywhere, Category = "Mesh Wrap", Meta = (ClampMin = "-1"))
+	int32 VertexIndex = INDEX_NONE;
+
+	bool operator==(const FMeshWrapToolLandmark& Other) const
+	{
+		return Identifier == Other.Identifier && VertexIndex == Other.VertexIndex;
+	}
+};
+
+
+UCLASS(MinimalAPI)
+class UMeshWrapLandmarkSelectionToolProperties : public UInteractiveToolPropertySet
+{
+	GENERATED_BODY()
+
+public:
+	/** Landmark names */
+	UPROPERTY(EditAnywhere, Category = Landmarks, EditFixedSize, Transient)
+	TArray<FMeshWrapToolLandmark> Landmarks;
+
+	/** Currently editable landmark. Set to -1 or hold the Shift key to add a new landmark. Hold the Ctrl key and select an existing landmark to set this value. */
+	UPROPERTY(EditAnywhere, Category = Landmarks, Meta = (ClampMin = "-1"))
+	int32 CurrentEditableLandmark = INDEX_NONE;
+
+	UPROPERTY(EditAnywhere, Category = Visualization, meta = (DisplayName = "Show Vertices"))
+	bool bShowVertices = false;
+
+	UPROPERTY(EditAnywhere, Category = Visualization, meta = (DisplayName = "Show Edges"))
+	bool bShowEdges = false;
+
+};
+
+UCLASS(MinimalAPI)
+class UMeshWrapLandmarkSelectionTool : public USingleSelectionMeshEditingTool
+{
+	GENERATED_BODY()
+
+public:
+
+	int32 GetFirstLandmarkWithID(const int32 Id) const;
+
+private:
+
+	friend class UMeshWrapLandmarkSelectionToolBuilder;
+
+	virtual void Setup() override;
+	virtual void OnShutdown(EToolShutdownType ShutdownType) override;
+	virtual void Render(IToolsContextRenderAPI* RenderAPI) override;
+	virtual void DrawHUD(FCanvas* Canvas, IToolsContextRenderAPI* RenderAPI) override;
+
+	virtual bool HasCancel() const override { return true; }
+	virtual bool HasAccept() const override { return true; }
+	virtual bool CanAccept() const override;
+
+	// IInteractiveToolCameraFocusAPI implementation
+	virtual FBox GetWorldSpaceFocusBox() override;
+
+	void SetDataflowContextObject(TObjectPtr<UDataflowContextObject> InDataflowContextObject);
+
+	void UpdateSelectedNode();
+	void InitializeSculptMeshFromTarget();
+	void SetPropertiesFromSelectedNode();
+	void SetCurrentEditableLandmark(int32 Index);
+	void UpdateLandmarkFromSelection();
+
+	UPROPERTY()
+	TObjectPtr<UMeshWrapLandmarkSelectionToolProperties> ToolProperties;
+
+	UPROPERTY()
+	TObjectPtr<UPreviewMesh> PreviewMesh = nullptr;
+
+	UPROPERTY()
+	TObjectPtr<UMeshWrapLandmarkSelectionMechanic> SelectionMechanic = nullptr;
+
+	UPROPERTY()
+	TObjectPtr<UDataflowContextObject> DataflowContextObject = nullptr;
+
+	TUniquePtr<UE::Geometry::FGroupTopology> Topology;
+
+	bool bAnyChangeMade = false;
+
+	FMeshWrapLandmarksNode* SelectionNodeToUpdate = nullptr;
+
+	FToolDataVisualizer LandmarkVisualizer;
+
+};
+
+
+UCLASS(MinimalAPI)
+class UMeshWrapLandmarkSelectionToolBuilder : public UInteractiveToolWithToolTargetsBuilder, public IDataflowEditorToolBuilder
+{
+	GENERATED_BODY()
+
+private:
+
+	// IDataflowEditorToolBuilder
+	UE_API virtual void GetSupportedConstructionViewModes(const UDataflowContextObject& ContextObject, TArray<const UE::Dataflow::IDataflowConstructionViewMode*>& Modes) const override;
+
+	// UInteractiveToolWithToolTargetsBuilder
+	UE_API virtual bool CanBuildTool(const FToolBuilderState& SceneState) const override;
+	UE_API virtual UInteractiveTool* BuildTool(const FToolBuilderState& SceneState) const override;
+	UE_API virtual const FToolTargetTypeRequirements& GetTargetRequirements() const override;
+};
+
+#undef UE_API
diff --git a/Engine/Plugins/Experimental/MeshResizing/Source/MeshResizingNodes/MeshResizingDataflowNodes.Build.cs b/Engine/Plugins/Experimental/MeshResizing/Source/MeshResizingNodes/MeshResizingDataflowNodes.Build.cs
--- a/Engine/Plugins/Experimental/MeshResizing/Source/MeshResizingNodes/MeshResizingDataflowNodes.Build.cs
+++ b/Engine/Plugins/Experimental/MeshResizing/Source/MeshResizingNodes/MeshResizingDataflowNodes.Build.cs
@@ -26,6 +26,7 @@ public class MeshResizingDataflowNodes : ModuleRules
 				"DataflowCore",
 				"DataflowEngine",
 				"DataflowNodes",
+				"DynamicMesh",
 				"GeometryFramework",
 				"MeshConversion",
 				"MeshConversionEngineTypes",
diff --git a/Engine/Plugins/Experimental/MeshResizing/Source/MeshResizingNodes/Private/MeshResizing/MeshResizingDataflowNodesModule.cpp b/Engine/Plugins/Experimental/MeshResizing/Source/MeshResizingNodes/Private/MeshResizing/MeshResizingDataflowNodesModule.cpp
--- a/Engine/Plugins/Experimental/MeshResizing/Source/MeshResizingNodes/Private/MeshResizing/MeshResizingDataflowNodesModule.cpp
+++ b/Engine/Plugins/Experimental/MeshResizing/Source/MeshResizingNodes/Private/MeshResizing/MeshResizingDataflowNodesModule.cpp
@@ -8,6 +8,7 @@
 #include "MeshResizing/MeshConstraintNodes.h"
 #include "MeshResizing/MeshResizingTextureNodes.h"
 #include "MeshResizing/MeshWarpNode.h"
+#include "MeshResizing/MeshWrapNode.h"
 #include "MeshResizing/RBFInterpolationNodes.h"
 #include "MeshResizing/UVUnwrapNode.h"
 #include "MeshResizing/UVMeshTransformNode.h"
@@ -30,6 +31,7 @@ namespace UE::MeshResizing
 
 			RegisterBaseBodyDataflowNodes();
 			RegisterMeshConstraintDataflowNodes();
+			RegisterMeshWrapNodes();
 			RegisterUVUnwrapNodes();
 			RegisterUVMeshTransformNodes();
 			RegisterAlignUVMeshNodes();
diff --git a/Engine/Plugins/Experimental/MeshResizing/Source/MeshResizingNodes/Private/MeshResizing/MeshWrapNode.cpp b/Engine/Plugins/Experimental/MeshResizing/Source/MeshResizingNodes/Private/MeshResizing/MeshWrapNode.cpp
new file mode 100644
--- /dev/null
+++ b/Engine/Plugins/Experimental/MeshResizing/Source/MeshResizingNodes/Private/MeshResizing/MeshWrapNode.cpp
@@ -0,0 +1,354 @@
+// Copyright Epic Games, Inc. All Rights Reserved.
+#include "MeshResizing/MeshWrapNode.h"
+#include "Components/DynamicMeshComponent.h"
+#include "DynamicMesh/DynamicMesh3.h"
+#include "UDynamicMesh.h"
+#include "Dataflow/DataflowMesh.h"
+#include "Dataflow/DataflowNodeColorsRegistry.h"
+#include "Dataflow/DataflowNodeFactory.h"
+#include "Dataflow/DataflowObject.h"
+#if WITH_EDITOR
+#include "Dataflow/DataflowDebugDrawInterface.h"
+#include "Dataflow/DataflowRenderingViewMode.h"
+#endif
+#include "Materials/Material.h"
+#include "Materials/MaterialInterface.h"
+
+#include UE_INLINE_GENERATED_CPP_BY_NAME(MeshWrapNode)
+
+using UE::Geometry::FDynamicMesh3;
+
+namespace UE::MeshResizing
+{
+	void RegisterMeshWrapNodes()
+	{
+		static const FLinearColor CDefaultNodeBodyTintColor = FLinearColor(0.f, 0.f, 0.f, 0.5f);
+		DATAFLOW_NODE_REGISTER_CREATION_FACTORY(FMeshWrapNode);
+		DATAFLOW_NODE_REGISTER_CREATION_FACTORY(FMeshWrapLandmarksNode);
+	}
+}
+
+namespace UE::MeshResizing::Private
+{
+#if WITH_EDITOR
+	static FLinearColor PseudoRandomColor(int32 NumColorRotations)
+	{
+		constexpr uint8 Spread = 157;  // Prime number that gives a good spread of colors without getting too similar as a rand might do.
+		uint8 Seed = Spread;
+		NumColorRotations = FMath::Abs(NumColorRotations);
+		for (int32 Rotation = 0; Rotation < NumColorRotations; ++Rotation)
+		{
+			Seed += Spread;
+		}
+		return FLinearColor::MakeFromHSV8(Seed, 180, 140);
+	}
+#endif
+
+	static TObjectPtr<UDynamicMeshComponent> MakeDynamicMeshComponent(const TObjectPtr<UDataflowMesh> Mesh, const FName& MeshName, const TObjectPtr<AActor> RootActor)
+	{
+		static UMaterial* DefaultMaterial = nullptr;
+		if (!DefaultMaterial)
+		{
+			// Use dataflow's default material
+			DefaultMaterial = Cast<UMaterial>(StaticLoadObject(UMaterial::StaticClass(), NULL, TEXT("/Engine/EditorMaterials/Dataflow/DataflowTwoSidedVertexMaterial")));
+		}
+
+		if (Mesh)
+		{
+			if (Mesh->GetDynamicMesh())
+			{
+				FDynamicMesh3 WrappedMeshCopy;
+				WrappedMeshCopy.Copy(Mesh->GetDynamicMeshRef());
+				TObjectPtr<UDynamicMeshComponent> DynamicMeshComponent = NewObject<UDynamicMeshComponent>(RootActor);
+				DynamicMeshComponent->SetMesh(MoveTemp(WrappedMeshCopy));
+				if (Mesh->GetMaterials().IsEmpty())
+				{
+					DynamicMeshComponent->SetOverrideRenderMaterial(DefaultMaterial);
+				}
+				else
+				{
+					DynamicMeshComponent->ConfigureMaterialSet(Mesh->GetMaterials());
+				}
+				return DynamicMeshComponent;
+			}
+		}
+		return nullptr;
+	}
+}
+
+FMeshWrapLandmarksNode::FMeshWrapLandmarksNode(const UE::Dataflow::FNodeParameters& InParam, FGuid InGuid) :
+	FDataflowNode(InParam, InGuid)
+{
+	RegisterInputConnection(&Mesh);
+	RegisterOutputConnection(&Mesh, &Mesh);
+	RegisterOutputConnection(&Landmarks);
+}
+
+void FMeshWrapLandmarksNode::Evaluate(UE::Dataflow::FContext& Context, const FDataflowOutput* Out) const
+{
+	if (Out->IsA(&Landmarks))
+	{
+		SetValue(Context, Landmarks, &Landmarks);
+	}
+	if (Out->IsA(&Mesh))
+	{
+		SafeForwardInput(Context, &Mesh, &Mesh);
+	}
+}
+
+#if WITH_EDITOR
+bool FMeshWrapLandmarksNode::CanDebugDrawViewMode(const FName& ViewMode) const
+{
+	return ViewMode == UE::Dataflow::FDataflowConstruction3DViewMode::Name;
+}
+
+void FMeshWrapLandmarksNode::DebugDraw(UE::Dataflow::FContext& Context, IDataflowDebugDrawInterface& DataflowRenderingInterface, const FDebugDrawParameters& DebugDrawParameters) const
+{
+	DataflowRenderingInterface.SetPointSize(PointSize);
+	if (bCanDebugDraw && (DebugDrawParameters.bNodeIsSelected || DebugDrawParameters.bNodeIsPinned))
+	{
+		if (TObjectPtr<UDataflowMesh> InMesh = GetValue(Context, &Mesh))
+		{
+			if (InMesh->GetDynamicMesh())
+			{
+				const TArray<FMeshWrapLandmark>& OutLandmarks = GetOutputValue(Context, &Landmarks, Landmarks);
+				DataflowRenderingInterface.ReservePoints(OutLandmarks.Num());
+
+				for (int32 Index = 0; Index < OutLandmarks.Num(); ++Index)
+				{
+					if (InMesh->GetDynamicMeshRef().IsVertex(OutLandmarks[Index].VertexIndex))
+					{
+						DataflowRenderingInterface.SetColor(UE::MeshResizing::Private::PseudoRandomColor(Index));
+						const FVector3d& Point = InMesh->GetDynamicMeshRef().GetVertexRef(OutLandmarks[Index].VertexIndex);
+						DataflowRenderingInterface.DrawPoint(Point);
+						if (bShowIndex || bShowIdentifier)
+						{
+							const FString IndexString = bShowIndex ? FString::FromInt(Index) : FString();
+
+							DataflowRenderingInterface.DrawText3d(IndexString + FString(" ") + OutLandmarks[Index].Identifier, Point);
+						}
+					}
+				}
+			}
+		}
+	}
+}
+#endif
+
+// Object encapsulating a change to the Selection Node's values. Used for Undo/Redo.
+class FMeshWrapLandmarksNode::FLandmarksNodeChange final : public FToolCommandChange
+{
+public:
+	FLandmarksNodeChange(const FMeshWrapLandmarksNode& Node)
+		: NodeGuid(Node.GetGuid())
+		, SavedLandmarks(Node.Landmarks)
+	{}
+
+private:
+	FGuid NodeGuid;
+	TArray<FMeshWrapLandmark> SavedLandmarks;
+
+	virtual FString ToString() const final { return TEXT("MeshWrapLandmarksNodeChange"); }
+
+	virtual void Apply(UObject* Object) final { SwapApplyRevert(Object); }
+	virtual void Revert(UObject* Object) final { SwapApplyRevert(Object); }
+
+	void SwapApplyRevert(UObject* Object)
+	{
+		if (UDataflow* const Dataflow = Cast<UDataflow>(Object))
+		{
+			if (const TSharedPtr<FDataflowNode> BaseNode = Dataflow->GetDataflow()->FindBaseNode(NodeGuid))
+			{
+				if (FMeshWrapLandmarksNode* const Node = BaseNode->AsType<FMeshWrapLandmarksNode>())
+				{
+					Swap(Node->Landmarks, SavedLandmarks);
+					Node->Invalidate();
+				}
+			}
+		}
+	}
+};
+
+TUniquePtr<class FToolCommandChange> FMeshWrapLandmarksNode::MakeSelectedNodeChange(const FMeshWrapLandmarksNode& Node)
+{
+	return MakeUnique<FMeshWrapLandmarksNode::FLandmarksNodeChange>(Node);
+}
+
+FMeshWrapNode::FMeshWrapNode(const UE::Dataflow::FNodeParameters& InParam, FGuid InGuid) :
+	FDataflowPrimitiveNode(InParam, InGuid)
+{
+	RegisterInputConnection(&SourceTopologyMesh);
+	RegisterInputConnection(&TargetShapeMesh);
+	RegisterInputConnection(&SourceTopologyLandmarks);
+	RegisterInputConnection(&TargetShapeLandmarks);
+
+	RegisterOutputConnection(&WrappedMesh, &SourceTopologyMesh);
+	RegisterOutputConnection(&MatchedLandmarks);
+}
+
+void FMeshWrapNode::Evaluate(UE::Dataflow::FContext& Context, const FDataflowOutput* Out) const
+{
+	using namespace UE::Geometry;
+
+	if (Out->IsA<TArray<FMeshWrapCorrespondence>>(&MatchedLandmarks))
+	{
+		SetValue(Context, CalculateMatchedLandmarks(Context), &MatchedLandmarks);
+		return;
+	}
+
+	if (Out->IsA(&WrappedMesh))
+	{
+		if (TObjectPtr<UDataflowMesh> InSourceTopologyMesh = GetValue(Context, &SourceTopologyMesh))
+		{
+			if (TObjectPtr<UDataflowMesh> InTargetShapeMesh = GetValue(Context, &TargetShapeMesh))
+			{
+				if (InSourceTopologyMesh->GetDynamicMesh() && InTargetShapeMesh->GetDynamicMesh())
+				{
+					const TArray<FMeshWrapCorrespondence>& InMatchedLandmarks = GetOutputValue(Context, &MatchedLandmarks, TArray<FMeshWrapCorrespondence>());
+					TArray<FWrapMeshCorrespondence> Correspondences;
+					Correspondences.Reserve(InMatchedLandmarks.Num());
+					for (const FMeshWrapCorrespondence& MatchedLandmark : InMatchedLandmarks)
+					{
+						Correspondences.Emplace(MatchedLandmark.SourceVertexIndex, MatchedLandmark.TargetVertexIndex);
+					}
+					TObjectPtr<UDataflowMesh> OutWrappedMesh = NewObject<UDataflowMesh>();
+					FDynamicMesh3 WrappedFMesh;
+					FWrapMesh MeshWrapper(nullptr);
+					MeshWrapper.MaxNumOuterIterations = MaxNumOuterIterations;
+					MeshWrapper.NumInnerIterations = NumInnerIterations;
+					MeshWrapper.ProjectionTolerance = ProjectionTolerance;
+					MeshWrapper.LaplacianStiffness = LaplacianStiffness;
+					MeshWrapper.InitialProjectionStiffness = InitialProjectionStiffness;
+					MeshWrapper.ProjectionStiffnessMuliplier = ProjectionStiffnessMuliplier;
+					MeshWrapper.CorrespondenceStiffness = CorrespondenceStiffness;
+					MeshWrapper.LaplacianType = FWrapMesh::ELaplacianType::AffineInvariant;
+					MeshWrapper.SetMesh(InSourceTopologyMesh->GetDynamicMesh());
+
+					MeshWrapper.WrapToTargetShape(InTargetShapeMesh->GetDynamicMeshRef(), Correspondences, WrappedFMesh);
+					OutWrappedMesh->SetDynamicMesh(MoveTemp(WrappedFMesh));
+					OutWrappedMesh->SetMaterials(InSourceTopologyMesh->GetMaterials());
+					SetValue(Context, OutWrappedMesh, &WrappedMesh);
+
+					return;
+				}
+			}
+
+		}
+
+		SafeForwardInput(Context, &SourceTopologyMesh, &WrappedMesh);
+	}
+}
+
+#if WITH_EDITOR
+bool FMeshWrapNode::CanDebugDrawViewMode(const FName& ViewMode) const
+{
+	return ViewMode == UE::Dataflow::FDataflowConstruction3DViewMode::Name;
+}
+
+void FMeshWrapNode::DebugDraw(UE::Dataflow::FContext& Context, IDataflowDebugDrawInterface& DataflowRenderingInterface, const FDebugDrawParameters& DebugDrawParameters) const
+{
+	constexpr float PointSize = 5.f;
+	DataflowRenderingInterface.SetPointSize(PointSize);
+	if (DebugDrawParameters.bNodeIsSelected || DebugDrawParameters.bNodeIsPinned)
+	{
+		if (TObjectPtr<UDataflowMesh> OutWrappedMesh = GetOutputValue(Context, &WrappedMesh, WrappedMesh))
+		{
+			if (TObjectPtr<UDataflowMesh> InTargetShapeMesh = GetValue(Context, &TargetShapeMesh))
+			{
+				if (TObjectPtr<UDataflowMesh> InSourceTopologyMesh = GetValue(Context, &SourceTopologyMesh))
+				{
+					if (OutWrappedMesh->GetDynamicMesh() && InTargetShapeMesh->GetDynamicMesh() && InSourceTopologyMesh->GetDynamicMesh())
+					{
+						const TArray<FMeshWrapCorrespondence>& InSourceTargetCorrespondences = GetOutputValue(Context, &MatchedLandmarks, TArray<FMeshWrapCorrespondence>());
+						DataflowRenderingInterface.ReservePoints((2 + (bDisplaySource ? 1 : 0) + (bDisplayTarget ? 1 : 0)) * InSourceTargetCorrespondences.Num());
+
+						for (int32 Index = 0; Index < InSourceTargetCorrespondences.Num(); ++Index)
+						{
+							const FMeshWrapCorrespondence& Correspondence = InSourceTargetCorrespondences[Index];
+							if (OutWrappedMesh->GetDynamicMeshRef().IsVertex(Correspondence.SourceVertexIndex) && InTargetShapeMesh->GetDynamicMeshRef().IsVertex(Correspondence.TargetVertexIndex) && InSourceTopologyMesh->GetDynamicMeshRef().IsVertex(Correspondence.SourceVertexIndex))
+							{
+								DataflowRenderingInterface.SetColor(UE::MeshResizing::Private::PseudoRandomColor(Index));
+								const FVector3d& WrappedPoint = OutWrappedMesh->GetDynamicMeshRef().GetVertexRef(Correspondence.SourceVertexIndex);
+								const FVector3d& TargetPoint = InTargetShapeMesh->GetDynamicMeshRef().GetVertexRef(Correspondence.TargetVertexIndex);
+								DataflowRenderingInterface.DrawPoint(WrappedPoint);
+								DataflowRenderingInterface.DrawPoint(TargetPoint);
+								DataflowRenderingInterface.DrawLine(WrappedPoint, TargetPoint);
+								const FString IndexString = FString::FromInt(Index) + FString(" ") + Correspondence.Identifier;
+								DataflowRenderingInterface.DrawText3d(IndexString, WrappedPoint);
+
+								if (bDisplaySource)
+								{
+									const FVector3d OffsetSourcePoint = InSourceTopologyMesh->GetDynamicMeshRef().GetVertexRef(Correspondence.SourceVertexIndex) + FVector3d(SourceDisplayOffset, 0., 0.);
+									DataflowRenderingInterface.DrawPoint(OffsetSourcePoint);
+									DataflowRenderingInterface.DrawText3d(IndexString, OffsetSourcePoint);
+								}
+
+								if(bDisplayTarget)
+								{
+									const FVector3d OffsetTargetPoint = TargetPoint + FVector3d(TargetDisplayOffset, 0., 0.);
+									DataflowRenderingInterface.DrawPoint(OffsetTargetPoint);
+									DataflowRenderingInterface.DrawText3d(IndexString, OffsetTargetPoint);
+								}
+							}
+						}
+					}
+				}
+			}
+		}
+	}
+}
+#endif
+
+void FMeshWrapNode::AddPrimitiveComponents(UE::Dataflow::FContext& Context, const TSharedPtr<const FManagedArrayCollection> RenderCollection, TObjectPtr<UObject> NodeOwner,
+	TObjectPtr<AActor> RootActor, TArray<TObjectPtr<UPrimitiveComponent>>& PrimitiveComponents)
+{
+	if (bDisplaySource)
+	{
+		static const FName SourceMeshName(TEXT("SourceTopologyMesh"));
+		if (TObjectPtr<UDynamicMeshComponent> SourceComp = UE::MeshResizing::Private::MakeDynamicMeshComponent(GetValue(Context, &SourceTopologyMesh), SourceMeshName, RootActor))
+		{
+			SourceComp->SetWorldTransform(FTransform(FVector3d(SourceDisplayOffset, 0., 0.)));
+			PrimitiveComponents.Add(SourceComp);
+		}
+	}
+
+	if (bDisplayTarget)
+	{
+		static const FName TargetMeshName(TEXT("TargetShapeMesh"));
+		if (TObjectPtr<UDynamicMeshComponent> TargetComp = UE::MeshResizing::Private::MakeDynamicMeshComponent(GetValue(Context, &TargetShapeMesh), TargetMeshName, RootActor))
+		{
+			TargetComp->SetWorldTransform(FTransform(FVector3d(TargetDisplayOffset, 0., 0.)));
+			PrimitiveComponents.Add(TargetComp);
+		}
+	}
+}
+
+TArray<FMeshWrapCorrespondence> FMeshWrapNode::CalculateMatchedLandmarks(UE::Dataflow::FContext& Context) const
+{
+	const TArray<FMeshWrapLandmark>& InSourceLandmarks = GetValue(Context, &SourceTopologyLandmarks);
+	const TArray<FMeshWrapLandmark>& InTargetLandmarks = GetValue(Context, &TargetShapeLandmarks);
+
+	// Convert Source to a Map for faster lookup
+	TMap<FString, int32> SourceLandmarkMap;
+	SourceLandmarkMap.Reserve(InSourceLandmarks.Num());
+	for (const FMeshWrapLandmark& SourceLandmark : InSourceLandmarks)
+	{
+		SourceLandmarkMap.Add(SourceLandmark.Identifier, SourceLandmark.VertexIndex);
+	}
+
+	TArray<FMeshWrapCorrespondence> Result;
+	Result.Reserve(InSourceLandmarks.Num());
+
+	for (const FMeshWrapLandmark& TargetLandmark : InTargetLandmarks)
+	{
+		if (const int32* SourceVertex = SourceLandmarkMap.Find(TargetLandmark.Identifier))
+		{
+			FMeshWrapCorrespondence& Correspondence = Result.AddDefaulted_GetRef();
+			Correspondence.Identifier = TargetLandmark.Identifier;
+			Correspondence.SourceVertexIndex = *SourceVertex;
+			Correspondence.TargetVertexIndex = TargetLandmark.VertexIndex;
+		}
+	}
+
+	return Result;
+}
\ No newline at end of file
diff --git a/Engine/Plugins/Experimental/MeshResizing/Source/MeshResizingNodes/Public/MeshResizing/MeshWrapNode.h b/Engine/Plugins/Experimental/MeshResizing/Source/MeshResizingNodes/Public/MeshResizing/MeshWrapNode.h
new file mode 100644
--- /dev/null
+++ b/Engine/Plugins/Experimental/MeshResizing/Source/MeshResizingNodes/Public/MeshResizing/MeshWrapNode.h
@@ -0,0 +1,227 @@
+// Copyright Epic Games, Inc. All Rights Reserved.
+
+#pragma once
+
+#include "Dataflow/DataflowPrimitiveNode.h"
+#include "GeometryCollection/ManagedArrayCollection.h"
+#include "Operations/WrapMesh.h"
+
+#include "MeshWrapNode.generated.h"
+
+#define UE_API MESHRESIZINGDATAFLOWNODES_API
+
+class UDataflowMesh;
+class UMaterialInterface;
+class UMaterial;
+
+/**
+ * Mesh Wrap Node landmark. Matched landmarks between source topology and target shape meshes are used to guide the Mesh Wrap operation.
+ */
+USTRUCT()
+struct FMeshWrapLandmark
+{
+	GENERATED_USTRUCT_BODY()
+
+	/** String name. Landmarks will be matched by comparing identifiers.*/
+	UPROPERTY(EditAnywhere, Category = "Mesh Wrap")
+	FString Identifier;
+
+	/** Vertex of this landmark. */
+	UPROPERTY(EditAnywhere, Category = "Mesh Wrap", meta = (ClampMin = "-1"))
+	int32 VertexIndex = INDEX_NONE;
+};
+
+/**
+ * Matched Mesh Wrap Node landmark correspondence.
+ */
+USTRUCT()
+struct FMeshWrapCorrespondence
+{
+	GENERATED_USTRUCT_BODY()
+
+	/** Matched landmark name.*/
+	UPROPERTY(EditAnywhere, Category = "Mesh Wrap")
+	FString Identifier;
+
+	/** Vertex on source topology mesh.*/
+	UPROPERTY(EditAnywhere, Category = "Mesh Wrap", meta = (ClampMin = "-1"))
+	int32 SourceVertexIndex = INDEX_NONE;
+
+	/** Vertex on target shape mesh. */
+	UPROPERTY(EditAnywhere, Category = "Mesh Wrap", meta = (ClampMin = "-1"))
+	int32 TargetVertexIndex = INDEX_NONE;
+};
+
+/** Node for defining landmarks used by MeshWrapNode. The Mesh Wrap Landmark Selection Tool allows generating these landmarks via selection. */
+USTRUCT(Meta = (MeshResizing, Experimental))
+struct FMeshWrapLandmarksNode : public FDataflowNode
+{
+	GENERATED_USTRUCT_BODY()
+	DATAFLOW_NODE_DEFINE_INTERNAL(FMeshWrapLandmarksNode, "MeshWrapLandmarks", "MeshResizing", "Mesh Wrap Landmarks")
+	DATAFLOW_NODE_RENDER_TYPE("SurfaceRender", FName("UDataflowMesh"), "Mesh")
+
+public:
+	FMeshWrapLandmarksNode(const UE::Dataflow::FNodeParameters& InParam, FGuid InGuid = FGuid::NewGuid());
+
+private:
+	//~ Begin FDataflowNode interface
+	virtual void Evaluate(UE::Dataflow::FContext& Context, const FDataflowOutput* Out) const override;
+#if WITH_EDITOR
+	virtual bool CanDebugDraw() const override
+	{
+		return bCanDebugDraw;
+	}
+	virtual bool CanDebugDrawViewMode(const FName& ViewModeName) const override;
+	virtual void DebugDraw(UE::Dataflow::FContext& Context, IDataflowDebugDrawInterface& DataflowRenderingInterface, const FDebugDrawParameters& DebugDrawParameters) const override;
+#endif
+	//~ End FDataflowNode interface
+
+	// Undo/redo handling for associated tool.
+	friend class UMeshWrapLandmarkSelectionTool;
+	class FLandmarksNodeChange;
+	static TUniquePtr<class FToolCommandChange> UE_API MakeSelectedNodeChange(const FMeshWrapLandmarksNode& Node);
+
+	/** The mesh to define landmarks on. */
+	UPROPERTY(meta = (DataflowInput, DataflowOutput, DataflowPassthrough = "Mesh"))
+	TObjectPtr<UDataflowMesh> Mesh;
+
+	/** The defined landmarks (identifier, vertex pair). Can be hand-input/edited or generated using the Mesh Wrap Landmark Selection Tool */
+	UPROPERTY(EditAnywhere, Category = "Mesh Wrap", meta = (DataflowOutput))
+	TArray<FMeshWrapLandmark> Landmarks;
+
+	/** Point size of displayed landmarks */
+	UPROPERTY(EditAnywhere, Category = "Display", meta = (ClampMin = 0.f))
+	float PointSize = 5.f;
+
+	/** Display the landmark indices */
+	UPROPERTY(EditAnywhere, Category = "Display")
+	bool bShowIndex = true;
+
+	/** Display the landmark identifier strings */
+	UPROPERTY(EditAnywhere, Category = "Display")
+	bool bShowIdentifier = true;
+
+	// This will be set by UMeshWrapLandmarkSelectionTool to false to not double draw.
+	bool bCanDebugDraw = true;
+};
+
+/** Dataflow node for wrapping one mesh's topology to another mesh's shape. Uses point landmarks defined by the Mesh Wrap Landmarks node to match corresponding points between the two meshes.*/
+USTRUCT(Meta = (MeshResizing, Experimental))
+struct FMeshWrapNode : public FDataflowPrimitiveNode
+{
+	GENERATED_USTRUCT_BODY()
+	DATAFLOW_NODE_DEFINE_INTERNAL(FMeshWrapNode, "MeshWrap", "MeshResizing", "Mesh Wrap")
+	DATAFLOW_NODE_RENDER_TYPE("SurfaceRender", FName("UDataflowMesh"), "WrappedMesh")
+
+public:
+
+	FMeshWrapNode(const UE::Dataflow::FNodeParameters& InParam, FGuid InGuid = FGuid::NewGuid());
+
+private:
+
+	/** Input mesh with the desired wrapped mesh topology. */
+	UPROPERTY(meta = (DataflowInput))
+	TObjectPtr<UDataflowMesh> SourceTopologyMesh;
+
+	/** Input mesh with the desired wrapped mesh shape. */
+	UPROPERTY(meta = (DataflowInput))
+	TObjectPtr<UDataflowMesh> TargetShapeMesh;
+
+	/** Output wrapped mesh. */
+	UPROPERTY(meta = (DataflowOutput, DataflowPassthrough = "SourceTopologyMesh"))
+	TObjectPtr<UDataflowMesh> WrappedMesh;
+
+	/** Landmarks defined on SourceTopologyMesh. TargetShapeLandmarks with matching Identifiers will be used to find correspondences that help improve the wrap. */
+	UPROPERTY(EditAnywhere, Category = "Mesh Wrap", meta = (DataflowInput))
+	TArray<FMeshWrapLandmark> SourceTopologyLandmarks;
+
+	/** Landmarks defined on TargetShapeMesh. SourceTopologyLandmarks with matching Identifiers will be used to find correspondences that help improve the wrap. */
+	UPROPERTY(EditAnywhere, Category = "Mesh Wrap", meta = (DataflowInput))
+	TArray<FMeshWrapLandmark> TargetShapeLandmarks;
+
+	/** Landmarks matched by Identifier from SourceTopologyLandmarks and TargetShapeLandmarks. */
+	UPROPERTY(meta = (DataflowOutput))
+	TArray<FMeshWrapCorrespondence> MatchedLandmarks;
+
+	/** Mesh Wrap is calculated with an inner and outer loop. This is the maximum number of outer loops. Each outer loop increases the Projection Stiffness by ProjectionStiffnessMultiplier.*/
+	UPROPERTY(EditAnywhere, Category = "Mesh Wrap", Meta = (ClampMin ="0"))
+	int32 MaxNumOuterIterations = 10;
+
+	/** Mesh Wrap is calculated with an inner and outer loop. This is the number of inner loops run before increasing the Projection Stiffness in the outer loop.*/
+	UPROPERTY(EditAnywhere, Category = "Mesh Wrap", Meta = (ClampMin = "0"))
+	int32 NumInnerIterations = 20;
+
+	/** Mesh Wrap will terminate early if the Projection tolerance is within this threshold.*/
+	UPROPERTY(EditAnywhere, Category = "Mesh Wrap", Meta =(ClampMin = "0"))
+	float ProjectionTolerance = 1e-4;
+
+	/** Weight of mesh wrap to retain Source Topology mesh features.*/
+	UPROPERTY(EditAnywhere, Category = "Mesh Wrap", AdvancedDisplay, Meta = (ClampMin = "0"))
+	float LaplacianStiffness = 1.;
+
+	/** Initial weight of mesh wrap to match projected Target Shape. Each outer loop will multiply this stiffness by ProjectionStiffnessMultiplier.*/
+	UPROPERTY(EditAnywhere, Category = "Mesh Wrap", AdvancedDisplay, Meta = (ClampMin = "0"))
+	float InitialProjectionStiffness = 0.1;
+	
+	/** Each outer loop will multiply InitialProjectionStiffness by this to improve Target Shape match.*/
+	UPROPERTY(EditAnywhere, Category = "Mesh Wrap", AdvancedDisplay, Meta = (ClampMin = "0"))
+	float ProjectionStiffnessMuliplier = 10.;
+
+	/** Weight of mesh wrap to match Landmark correspondences.*/
+	UPROPERTY(EditAnywhere, Category = "Mesh Wrap", AdvancedDisplay, Meta = (ClampMin = "0"))
+	float CorrespondenceStiffness = 1.;
+
+	/** Display material for Source or Target when none is supplied. */
+	UPROPERTY(EditAnywhere, Category = "Mesh Wrap Display")
+	TObjectPtr<UMaterial> DefaultDisplayMaterial;
+
+	/** Display landmarks. */
+	UPROPERTY(EditAnywhere, Category = "Mesh Wrap Display")
+	bool bDisplayLandmarks = true;
+
+	/** Draw source mesh. */
+	UPROPERTY(EditAnywhere, Category ="Mesh Wrap Display")
+	bool bDisplaySource = false;
+
+	/** Offset of source mesh display for side-by-side drawing. */
+	UPROPERTY(EditAnywhere, Category = "Mesh Wrap Display", Meta =(EditCondition = "bDisplaySource"))
+	float SourceDisplayOffset = -150.f;
+
+	/** Draw target mesh. */
+	UPROPERTY(EditAnywhere, Category = "Mesh Wrap Display")
+	bool bDisplayTarget = false;
+
+	/** Offset of target mesh display for side-by-side drawing. */
+	UPROPERTY(EditAnywhere, Category = "Mesh Wrap Display", Meta = (EditCondition = "bDisplayTarget"))
+	float TargetDisplayOffset = 150.f;
+
+	//~ Begin FDataflowNode interface
+	virtual void Evaluate(UE::Dataflow::FContext& Context, const FDataflowOutput* Out) const override;
+#if WITH_EDITOR
+	virtual bool CanDebugDraw() const override
+	{
+		return true;
+	}
+	virtual bool CanDebugDrawViewMode(const FName& ViewModeName) const override;
+	virtual void DebugDraw(UE::Dataflow::FContext& Context, IDataflowDebugDrawInterface& DataflowRenderingInterface, const FDebugDrawParameters& DebugDrawParameters) const override;
+#endif
+	//~ End FDataflowNode interface
+
+	//~ Begin FDataflowPrimitiveNode interface 
+	virtual bool HasRenderCollectionPrimitives() const override
+	{
+		return false;
+	}
+	virtual void AddPrimitiveComponents(UE::Dataflow::FContext& Context, const TSharedPtr<const FManagedArrayCollection> RenderCollection, TObjectPtr<UObject> NodeOwner,
+		TObjectPtr<AActor> RootActor, TArray<TObjectPtr<UPrimitiveComponent>>& PrimitiveComponents) override;
+	//~ End FDataflowPrimitiveNode interface
+
+	TArray<FMeshWrapCorrespondence> CalculateMatchedLandmarks(UE::Dataflow::FContext& Context) const;
+};
+
+namespace UE::MeshResizing
+{
+	void RegisterMeshWrapNodes();
+}
+
+#undef UE_API
\ No newline at end of file
