Fix calculation for RBF parallel step, so that we try to spawn one task per thread instead of one task per vertex

diff --git a/Engine/Plugins/Experimental/MeshResizing/Source/MeshResizingCore/Private/MeshResizing/RBFInterpolation.cpp b/Engine/Plugins/Experimental/MeshResizing/Source/MeshResizingCore/Private/MeshResizing/RBFInterpolation.cpp
--- a/Engine/Plugins/Experimental/MeshResizing/Source/MeshResizingCore/Private/MeshResizing/RBFInterpolation.cpp
+++ b/Engine/Plugins/Experimental/MeshResizing/Source/MeshResizingCore/Private/MeshResizing/RBFInterpolation.cpp
@@ -113,7 +113,7 @@ namespace UE::MeshResizing
 			};
 
 			const int32 NumVertexIndices = DeformingMesh.VertexCount();
-			const int32 NumTasks = FMath::Max(FTaskGraphInterface::Get().GetNumWorkerThreads(), NumVertexIndices);
+			const int32 NumTasks = FMath::Min(FTaskGraphInterface::Get().GetNumWorkerThreads(), NumVertexIndices);
 			constexpr int32 MinVerticesByTask = 10;
 			const int32 VerticesByTask = FMath::Max(FMath::Min(FMath::DivideAndRoundUp(NumVertexIndices, NumTasks), MinVerticesByTask), 1);
 			const int32 NumBatches = FMath::DivideAndRoundUp(NumVertexIndices, VerticesByTask);
@@ -168,7 +168,7 @@ namespace UE::MeshResizing
 				};
 
 			const int32 NumVertexIndices = Positions.Num();
-			const int32 NumTasks = FMath::Max(FTaskGraphInterface::Get().GetNumWorkerThreads(), NumVertexIndices);
+			const int32 NumTasks = FMath::Min(FTaskGraphInterface::Get().GetNumWorkerThreads(), NumVertexIndices);
 			constexpr int32 MinVerticesByTask = 10;
 			const int32 VerticesByTask = FMath::Max(FMath::Min(FMath::DivideAndRoundUp(NumVertexIndices, NumTasks), MinVerticesByTask), 1);
 			const int32 NumBatches = FMath::DivideAndRoundUp(NumVertexIndices, VerticesByTask);
