Dataflow Landmark tool : fix the tool marking the node invalid even when no changes were made

diff --git a/Engine/Plugins/Experimental/MeshResizing/Source/MeshResizingEditorTools/Private/MeshResizing/MeshWrapLandmarkSelectionTool.cpp b/Engine/Plugins/Experimental/MeshResizing/Source/MeshResizingEditorTools/Private/MeshResizing/MeshWrapLandmarkSelectionTool.cpp
--- a/Engine/Plugins/Experimental/MeshResizing/Source/MeshResizingEditorTools/Private/MeshResizing/MeshWrapLandmarkSelectionTool.cpp
+++ b/Engine/Plugins/Experimental/MeshResizing/Source/MeshResizingEditorTools/Private/MeshResizing/MeshWrapLandmarkSelectionTool.cpp
@@ -233,10 +233,38 @@ void UMeshWrapLandmarkSelectionTool::Setup()
 	// Initialize the Selection from the selected Dataflow node
 	SetPropertiesFromSelectedNode();
 	SelectionNodeToUpdate->bCanDebugDraw = false; // Disable node's debug draw to not double draw
-	SelectionNodeToUpdate->Invalidate();
-	ToolProperties->WatchProperty(ToolProperties->Landmarks, [this](const TArray<FMeshWrapToolLandmark>& NewLandmarks)
+	//SelectionNodeToUpdate->Invalidate();
+
+	ToolProperties->WatchProperty(ToolProperties->Landmarks,
+		[this](const TArray<FMeshWrapToolLandmark>& NewLandmarks)
 		{
-			bAnyChangeMade = true;
+			auto SameLandmarkAsNode = [this]() -> bool
+				{
+					if (SelectionNodeToUpdate == nullptr || ToolProperties == nullptr)
+					{
+						return false;
+					}
+					if (ToolProperties->Landmarks.Num() != SelectionNodeToUpdate->Landmarks.Num())
+					{
+						return false;
+					}
+					const int32 NumLandmarks = ToolProperties->Landmarks.Num();
+					for (int32 Idx = 0; Idx < NumLandmarks; ++Idx)
+					{
+						const FMeshWrapToolLandmark& ToolLandmark = ToolProperties->Landmarks[Idx];
+						const FMeshWrapLandmark& NodeLandmark = SelectionNodeToUpdate->Landmarks[Idx];
+						if (ToolLandmark.Identifier != NodeLandmark.Identifier)
+						{
+							return false;
+						}
+						if (ToolLandmark.VertexIndex != NodeLandmark.VertexIndex)
+						{
+							return false;
+						}
+					}
+					return true;
+				};
+			bAnyChangeMade = !SameLandmarkAsNode();
 		});
 	ToolProperties->WatchProperty(ToolProperties->CurrentEditableLandmark, [this](int32 NewCurrentEditableLandmark)
 		{
