Dataflow landmark tool : fix some pointer checks

diff --git a/Engine/Plugins/Experimental/MeshResizing/Source/MeshResizingEditorTools/Private/MeshResizing/MeshWrapLandmarkSelectionTool.cpp b/Engine/Plugins/Experimental/MeshResizing/Source/MeshResizingEditorTools/Private/MeshResizing/MeshWrapLandmarkSelectionTool.cpp
--- a/Engine/Plugins/Experimental/MeshResizing/Source/MeshResizingEditorTools/Private/MeshResizing/MeshWrapLandmarkSelectionTool.cpp
+++ b/Engine/Plugins/Experimental/MeshResizing/Source/MeshResizingEditorTools/Private/MeshResizing/MeshWrapLandmarkSelectionTool.cpp
@@ -232,8 +232,10 @@ void UMeshWrapLandmarkSelectionTool::Setup()
 	 
 	// Initialize the Selection from the selected Dataflow node
 	SetPropertiesFromSelectedNode();
-	SelectionNodeToUpdate->bCanDebugDraw = false; // Disable node's debug draw to not double draw
-	//SelectionNodeToUpdate->Invalidate();
+	if (SelectionNodeToUpdate)
+	{
+		SelectionNodeToUpdate->bCanDebugDraw = false; // Disable node's debug draw to not double draw
+	}
 
 	ToolProperties->WatchProperty(ToolProperties->Landmarks,
 		[this](const TArray<FMeshWrapToolLandmark>& NewLandmarks)
@@ -276,14 +278,18 @@ void UMeshWrapLandmarkSelectionTool::Setup()
 
 void UMeshWrapLandmarkSelectionTool::SetPropertiesFromSelectedNode()
 {
+	ToolProperties->Landmarks.Reset();
+
 	SelectionNodeToUpdate = DataflowContextObject->GetSelectedNodeOfType<FMeshWrapLandmarksNode>();
-	check(SelectionNodeToUpdate);
-	ToolProperties->Landmarks.SetNum(SelectionNodeToUpdate->Landmarks.Num());
-	for (int32 Index = 0; Index < SelectionNodeToUpdate->Landmarks.Num(); ++Index)
+	if (ensure(SelectionNodeToUpdate))
 	{
-		ToolProperties->Landmarks[Index].Identifier = SelectionNodeToUpdate->Landmarks[Index].Identifier;
-		ToolProperties->Landmarks[Index].VertexIndex = SelectionNodeToUpdate->Landmarks[Index].VertexIndex;
+		ToolProperties->Landmarks.SetNum(SelectionNodeToUpdate->Landmarks.Num());
+		for (int32 Index = 0; Index < SelectionNodeToUpdate->Landmarks.Num(); ++Index)
+		{
+			ToolProperties->Landmarks[Index].Identifier = SelectionNodeToUpdate->Landmarks[Index].Identifier;
+			ToolProperties->Landmarks[Index].VertexIndex = SelectionNodeToUpdate->Landmarks[Index].VertexIndex;
 
+		}
 	}
 	ToolProperties->CurrentEditableLandmark = INDEX_NONE;
 }
@@ -378,7 +384,10 @@ void UMeshWrapLandmarkSelectionTool::OnShutdown(EToolShutdownType ShutdownType)
 		UpdateSelectedNode();
 		GetToolManager()->EndUndoTransaction();
 
-		SelectionNodeToUpdate->Invalidate();
+		if (SelectionNodeToUpdate)
+		{
+			SelectionNodeToUpdate->Invalidate();
+		}
 	}
 
 	SelectionMechanic->Properties->SaveProperties(this);
@@ -394,7 +403,10 @@ void UMeshWrapLandmarkSelectionTool::OnShutdown(EToolShutdownType ShutdownType)
 	SelectionMechanic->Shutdown();
 
 	Topology.Reset();
-	SelectionNodeToUpdate->bCanDebugDraw = true;
+	if (SelectionNodeToUpdate)
+	{
+		SelectionNodeToUpdate->bCanDebugDraw = true;
+	}
 }
 
 void UMeshWrapLandmarkSelectionTool::Render(IToolsContextRenderAPI* RenderAPI)
@@ -464,22 +476,22 @@ void UMeshWrapLandmarkSelectionTool::SetDataflowContextObject(TObjectPtr<UDatafl
 
 void UMeshWrapLandmarkSelectionTool::UpdateSelectedNode()
 {
-	checkf(SelectionNodeToUpdate, TEXT("Expected non-null pointer to Selection Node"));
-
-	// Save previous state for undo
-	if (UDataflow* const Dataflow = DataflowContextObject->GetDataflowAsset())
+	if (ensureMsgf(SelectionNodeToUpdate, TEXT("Expected non-null pointer to Selection Node")))
 	{
-		GetToolManager()->GetContextTransactionsAPI()->AppendChange(Dataflow, 
-			FMeshWrapLandmarksNode::MakeSelectedNodeChange(*SelectionNodeToUpdate),
-			LOCTEXT("MeshWrapLandmarkUpdate", "Update Mesh Wrap Landmarks Node"));
-	}
+		// Save previous state for undo
+		if (UDataflow* const Dataflow = DataflowContextObject->GetDataflowAsset())
+		{
+			GetToolManager()->GetContextTransactionsAPI()->AppendChange(Dataflow,
+				FMeshWrapLandmarksNode::MakeSelectedNodeChange(*SelectionNodeToUpdate),
+				LOCTEXT("MeshWrapLandmarkUpdate", "Update Mesh Wrap Landmarks Node"));
+		}
 
-	check(SelectionNodeToUpdate);
-	SelectionNodeToUpdate->Landmarks.SetNum(ToolProperties->Landmarks.Num());
-	for (int32 Index = 0; Index < ToolProperties->Landmarks.Num(); ++Index)
-	{
-		SelectionNodeToUpdate->Landmarks[Index].Identifier = ToolProperties->Landmarks[Index].Identifier;
-		SelectionNodeToUpdate->Landmarks[Index].VertexIndex = ToolProperties->Landmarks[Index].VertexIndex;
+		SelectionNodeToUpdate->Landmarks.SetNum(ToolProperties->Landmarks.Num());
+		for (int32 Index = 0; Index < ToolProperties->Landmarks.Num(); ++Index)
+		{
+			SelectionNodeToUpdate->Landmarks[Index].Identifier = ToolProperties->Landmarks[Index].Identifier;
+			SelectionNodeToUpdate->Landmarks[Index].VertexIndex = ToolProperties->Landmarks[Index].VertexIndex;
+		}
 	}
 	
 }
