[MeshResizing] Fixed vertices per task alocation bug
- The amount of vertices per task was being computed in such a way so the amount of vertices per task could never be bigger than the MinVerticesByTask (it has the value of 10)

#rnx

diff --git a/Engine/Plugins/Experimental/MeshResizing/Source/MeshResizingCore/Private/MeshResizing/RBFInterpolation.cpp b/Engine/Plugins/Experimental/MeshResizing/Source/MeshResizingCore/Private/MeshResizing/RBFInterpolation.cpp
--- a/Engine/Plugins/Experimental/MeshResizing/Source/MeshResizingCore/Private/MeshResizing/RBFInterpolation.cpp
+++ b/Engine/Plugins/Experimental/MeshResizing/Source/MeshResizingCore/Private/MeshResizing/RBFInterpolation.cpp
@@ -115,7 +115,7 @@ namespace UE::MeshResizing
 			const int32 NumVertexIndices = DeformingMesh.VertexCount();
 			const int32 NumTasks = FMath::Max(FMath::Min(FTaskGraphInterface::Get().GetNumWorkerThreads(), NumVertexIndices), 1);
 			constexpr int32 MinVerticesByTask = 10;
-			const int32 VerticesByTask = FMath::Max(FMath::Min(FMath::DivideAndRoundUp(NumVertexIndices, NumTasks), MinVerticesByTask), 1);
+			const int32 VerticesByTask = FMath::Max(FMath::Max(FMath::DivideAndRoundUp(NumVertexIndices, NumTasks), MinVerticesByTask), 1);
 			const int32 NumBatches = FMath::DivideAndRoundUp(NumVertexIndices, VerticesByTask);
 			TArray<UE::Tasks::FTask> PendingTasks;
 			PendingTasks.Reserve(NumBatches);
