Adding includes before upcoming header cleanup.

#Synced-CL 51781844

diff --git a/Engine/Plugins/Compositing/CompositeCore/Source/CompositeCore/Private/CompositeCoreModule.cpp b/Engine/Plugins/Compositing/CompositeCore/Source/CompositeCore/Private/CompositeCoreModule.cpp
--- a/Engine/Plugins/Compositing/CompositeCore/Source/CompositeCore/Private/CompositeCoreModule.cpp
+++ b/Engine/Plugins/Compositing/CompositeCore/Source/CompositeCore/Private/CompositeCoreModule.cpp
@@ -3,6 +3,7 @@
 #include "CompositeCoreModule.h"
 
 #include "Interfaces/IPluginManager.h"
+#include "HAL/IConsoleManager.h"
 #include "HDRHelper.h"
 #include "Misc/ConfigUtilities.h"
 #include "Misc/Paths.h"
diff --git a/Engine/Plugins/Compositing/OpenColorIO/Source/OpenColorIO/Private/OpenColorIOModule.cpp b/Engine/Plugins/Compositing/OpenColorIO/Source/OpenColorIO/Private/OpenColorIOModule.cpp
--- a/Engine/Plugins/Compositing/OpenColorIO/Source/OpenColorIO/Private/OpenColorIOModule.cpp
+++ b/Engine/Plugins/Compositing/OpenColorIO/Source/OpenColorIO/Private/OpenColorIOModule.cpp
@@ -5,6 +5,7 @@
 #include "Interfaces/IPluginManager.h"
 #include "OpenColorIODisplayManager.h"
 #include "ShaderCore.h"
+#include "Misc/Paths.h"
 
 #if WITH_EDITOR
 #include "OpenColorIOWrapperModule.h"
diff --git a/Engine/Plugins/Compositing/OpenColorIO/Source/OpenColorIO/Public/OpenColorIOShared.h b/Engine/Plugins/Compositing/OpenColorIO/Source/OpenColorIO/Public/OpenColorIOShared.h
--- a/Engine/Plugins/Compositing/OpenColorIO/Source/OpenColorIO/Public/OpenColorIOShared.h
+++ b/Engine/Plugins/Compositing/OpenColorIO/Source/OpenColorIO/Public/OpenColorIOShared.h
@@ -11,6 +11,7 @@
 #include "Misc/SecureHash.h"
 #include "RenderResource.h"
 #include "RenderDeferredCleanup.h"
+#include "RenderingThread.h"
 #include "RHIFeatureLevel.h"
 #include "SceneTypes.h"
 #include "Shader.h"
diff --git a/Engine/Plugins/Experimental/MeshPartition/Source/MeshPartitionCompute/Private/MeshPartitionComputeModule.cpp b/Engine/Plugins/Experimental/MeshPartition/Source/MeshPartitionCompute/Private/MeshPartitionComputeModule.cpp
--- a/Engine/Plugins/Experimental/MeshPartition/Source/MeshPartitionCompute/Private/MeshPartitionComputeModule.cpp
+++ b/Engine/Plugins/Experimental/MeshPartition/Source/MeshPartitionCompute/Private/MeshPartitionComputeModule.cpp
@@ -3,6 +3,7 @@
 #include "MeshPartitionComputeModule.h"
 
 #include "Interfaces/IPluginManager.h"
+#include "Misc/Paths.h"
 #include "Modules/ModuleManager.h" // IMPLEMENT_MODULE
 #include "ShaderCore.h" // AddShaderSourceDirectoryMapping
 
diff --git a/Engine/Plugins/Experimental/MeshResizing/Source/MeshResizingNodes/Private/MeshResizing/UVResizeControllerNode.cpp b/Engine/Plugins/Experimental/MeshResizing/Source/MeshResizingNodes/Private/MeshResizing/UVResizeControllerNode.cpp
--- a/Engine/Plugins/Experimental/MeshResizing/Source/MeshResizingNodes/Private/MeshResizing/UVResizeControllerNode.cpp
+++ b/Engine/Plugins/Experimental/MeshResizing/Source/MeshResizingNodes/Private/MeshResizing/UVResizeControllerNode.cpp
@@ -1,10 +1,12 @@
 // Copyright Epic Games, Inc. All Rights Reserved.
+
 #include "MeshResizing/UVResizeControllerNode.h"
 #include "Dataflow/DataflowMesh.h"
 #include "DynamicMesh/DynamicMeshAttributeSet.h"
 #include "Engine/Texture.h"
 #include "Materials/MaterialInstanceConstant.h"
 #include "Misc/PackageName.h"
+#include "Misc/Paths.h"
 
 #include UE_INLINE_GENERATED_CPP_BY_NAME(UVResizeControllerNode)
 
diff --git a/Engine/Plugins/Experimental/WaterAdvanced/Source/WaterAdvanced/Private/WaterAdvanced.cpp b/Engine/Plugins/Experimental/WaterAdvanced/Source/WaterAdvanced/Private/WaterAdvanced.cpp
--- a/Engine/Plugins/Experimental/WaterAdvanced/Source/WaterAdvanced/Private/WaterAdvanced.cpp
+++ b/Engine/Plugins/Experimental/WaterAdvanced/Source/WaterAdvanced/Private/WaterAdvanced.cpp
@@ -4,6 +4,7 @@
 #include "ShallowWaterRiverDetails.h"
 #include "ShallowWaterRiverActor.h"
 #include "Interfaces/IPluginManager.h"
+#include "Misc/Paths.h"
 
 #if WITH_EDITOR
 #include "PropertyCustomizationHelpers.h"
diff --git a/Engine/Plugins/Interchange/Assets/Source/Assets/Private/InterchangeAssetsModule.cpp b/Engine/Plugins/Interchange/Assets/Source/Assets/Private/InterchangeAssetsModule.cpp
--- a/Engine/Plugins/Interchange/Assets/Source/Assets/Private/InterchangeAssetsModule.cpp
+++ b/Engine/Plugins/Interchange/Assets/Source/Assets/Private/InterchangeAssetsModule.cpp
@@ -5,6 +5,7 @@
 #include "Modules/ModuleManager.h"
 #include "ShaderCore.h"
 #include "Interfaces/IPluginManager.h"
+#include "Misc/Paths.h"
 
 class FInterchangeAssetsModule : public IModuleInterface
 {
diff --git a/Engine/Plugins/Performance/PerformanceMonitor/Source/PerformanceMonitor/Private/PerformanceMonitor.cpp b/Engine/Plugins/Performance/PerformanceMonitor/Source/PerformanceMonitor/Private/PerformanceMonitor.cpp
--- a/Engine/Plugins/Performance/PerformanceMonitor/Source/PerformanceMonitor/Private/PerformanceMonitor.cpp
+++ b/Engine/Plugins/Performance/PerformanceMonitor/Source/PerformanceMonitor/Private/PerformanceMonitor.cpp
@@ -9,6 +9,7 @@
 #include "Engine/GameViewportClient.h"
 #include "Misc/ConfigCacheIni.h"
 #include "Misc/AutomationTest.h"
+#include "RenderTimer.h"
 #include "Stats/StatsData.h"
 #include "UnrealClient.h"
 #include "DynamicRHI.h"
diff --git a/Engine/Plugins/Runtime/ComputeFramework/Source/ComputeFramework/Private/ComputeGraphInstance.cpp b/Engine/Plugins/Runtime/ComputeFramework/Source/ComputeFramework/Private/ComputeGraphInstance.cpp
--- a/Engine/Plugins/Runtime/ComputeFramework/Source/ComputeFramework/Private/ComputeGraphInstance.cpp
+++ b/Engine/Plugins/Runtime/ComputeFramework/Source/ComputeFramework/Private/ComputeGraphInstance.cpp
@@ -10,6 +10,7 @@
 #include "ComputeFramework/ComputeGraphRenderProxy.h"
 #include "ComputeFramework/ComputeGraphWorker.h"
 #include "ComputeFramework/ComputeSystem.h"
+#include "RenderingThread.h"
 
 #include UE_INLINE_GENERATED_CPP_BY_NAME(ComputeGraphInstance)
 
diff --git a/Engine/Plugins/Runtime/OpenXR/Source/OpenXRHMD/Private/OpenXRHMD.cpp b/Engine/Plugins/Runtime/OpenXR/Source/OpenXRHMD/Private/OpenXRHMD.cpp
--- a/Engine/Plugins/Runtime/OpenXR/Source/OpenXRHMD/Private/OpenXRHMD.cpp
+++ b/Engine/Plugins/Runtime/OpenXR/Source/OpenXRHMD/Private/OpenXRHMD.cpp
@@ -4,6 +4,7 @@
 
 #include "AnalyticsEventAttribute.h"
 #include "ClearQuad.h"
+#include "CommonRenderResources.h"
 #include "Containers/ArrayView.h"
 #include "DataDrivenShaderPlatformInfo.h"
 #include "DefaultStereoLayers.h"
diff --git a/Engine/Plugins/Runtime/WarpUtils/Source/PFMExporter/Private/PFMExporterModule.cpp b/Engine/Plugins/Runtime/WarpUtils/Source/PFMExporter/Private/PFMExporterModule.cpp
--- a/Engine/Plugins/Runtime/WarpUtils/Source/PFMExporter/Private/PFMExporterModule.cpp
+++ b/Engine/Plugins/Runtime/WarpUtils/Source/PFMExporter/Private/PFMExporterModule.cpp
@@ -7,6 +7,7 @@
 
 #include "Interfaces/IPluginManager.h"
 #include "Misc/Paths.h"
+#include "RenderingThread.h"
 #include "ShaderCore.h"
 #include "RHICommandList.h"
 
diff --git a/Engine/Plugins/TextureGraph/Source/TextureGraphEngine/Device/FX/Device_FX.cpp b/Engine/Plugins/TextureGraph/Source/TextureGraphEngine/Device/FX/Device_FX.cpp
--- a/Engine/Plugins/TextureGraph/Source/TextureGraphEngine/Device/FX/Device_FX.cpp
+++ b/Engine/Plugins/TextureGraph/Source/TextureGraphEngine/Device/FX/Device_FX.cpp
@@ -1,5 +1,6 @@
 // Copyright Epic Games, Inc. All Rights Reserved.
 #include "Device_FX.h"
+#include "CommonRenderResources.h"
 #include "TextureGraphEngine.h"
 #include "DeviceBuffer_FX.h"
 #include "Device/Mem/Device_Mem.h"
diff --git a/Engine/Plugins/TextureGraph/Source/TextureGraphEngine/FxMat/FxMaterial.cpp b/Engine/Plugins/TextureGraph/Source/TextureGraphEngine/FxMat/FxMaterial.cpp
--- a/Engine/Plugins/TextureGraph/Source/TextureGraphEngine/FxMat/FxMaterial.cpp
+++ b/Engine/Plugins/TextureGraph/Source/TextureGraphEngine/FxMat/FxMaterial.cpp
@@ -1,5 +1,6 @@
 // Copyright Epic Games, Inc. All Rights Reserved.
 #include "FxMaterial.h"
+#include "CommonRenderResources.h"
 #include "ShaderParameters.h" 
 #include "Device/Device.h"
 #include "EngineModule.h"
diff --git a/Engine/Plugins/VirtualProduction/Avalanche/Source/AvalancheMaterial/Private/MaterialCache/AvaMaterialCacheSettings.cpp b/Engine/Plugins/VirtualProduction/Avalanche/Source/AvalancheMaterial/Private/MaterialCache/AvaMaterialCacheSettings.cpp
--- a/Engine/Plugins/VirtualProduction/Avalanche/Source/AvalancheMaterial/Private/MaterialCache/AvaMaterialCacheSettings.cpp
+++ b/Engine/Plugins/VirtualProduction/Avalanche/Source/AvalancheMaterial/Private/MaterialCache/AvaMaterialCacheSettings.cpp
@@ -2,6 +2,7 @@
 
 #include "MaterialCache/AvaMaterialCacheSettings.h"
 #include "MaterialCache/AvaMaterialCacheLog.h"
+#include "Shader.h"
 #include "VertexFactory.h"
 
 const FLazyName UAvaMaterialCacheSettings::DefaultRealtimeProfile(TEXT("BasicSceneRealtime"));
diff --git a/Engine/Plugins/VirtualProduction/DMX/DMXPixelMapping/Source/DMXPixelMappingRenderer/Private/DMXPixelMappingRenderPixelMapProxy.cpp b/Engine/Plugins/VirtualProduction/DMX/DMXPixelMapping/Source/DMXPixelMappingRenderer/Private/DMXPixelMappingRenderPixelMapProxy.cpp
--- a/Engine/Plugins/VirtualProduction/DMX/DMXPixelMapping/Source/DMXPixelMappingRenderer/Private/DMXPixelMappingRenderPixelMapProxy.cpp
+++ b/Engine/Plugins/VirtualProduction/DMX/DMXPixelMapping/Source/DMXPixelMappingRenderer/Private/DMXPixelMappingRenderPixelMapProxy.cpp
@@ -2,6 +2,7 @@
 
 #include "DMXPixelMappingRenderPixelMapProxy.h"
 
+#include "CommonRenderResources.h"
 #include "DMXPixelMappingRenderElement.h"
 #include "Engine/TextureRenderTarget2D.h"
 #include "PixelShaderUtils.h"
diff --git a/Engine/Plugins/VirtualProduction/GeometryMask/Source/GeometryMask/Private/GeometryMaskModule.cpp b/Engine/Plugins/VirtualProduction/GeometryMask/Source/GeometryMask/Private/GeometryMaskModule.cpp
--- a/Engine/Plugins/VirtualProduction/GeometryMask/Source/GeometryMask/Private/GeometryMaskModule.cpp
+++ b/Engine/Plugins/VirtualProduction/GeometryMask/Source/GeometryMask/Private/GeometryMaskModule.cpp
@@ -1,8 +1,9 @@
-﻿// Copyright Epic Games, Inc. All Rights Reserved.
+// Copyright Epic Games, Inc. All Rights Reserved.
 
 #include "GeometryMaskModule.h"
 
 #include "GeometryMaskSettings.h"
+#include "HAL/IConsoleManager.h"
 #include "IGeometryMaskModule.h"
 #include "Interfaces/IPluginManager.h"
 #include "Misc/Paths.h"
diff --git a/Engine/Source/Developer/Apple/MetalShaderFormat/Private/MetalShaderCompiler.cpp b/Engine/Source/Developer/Apple/MetalShaderFormat/Private/MetalShaderCompiler.cpp
--- a/Engine/Source/Developer/Apple/MetalShaderFormat/Private/MetalShaderCompiler.cpp
+++ b/Engine/Source/Developer/Apple/MetalShaderFormat/Private/MetalShaderCompiler.cpp
@@ -2,6 +2,7 @@
 
 #include "MetalShaderCompiler.h"
 
+#include "Containers/AnsiString.h"
 #include "Containers/ArrayView.h"
 #include "HAL/PlatformFileManager.h"
 #include "MetalShaderFormat.h"
diff --git a/Engine/Source/Developer/Apple/MetalShaderFormat/Private/MetalShaderFormat.cpp b/Engine/Source/Developer/Apple/MetalShaderFormat/Private/MetalShaderFormat.cpp
--- a/Engine/Source/Developer/Apple/MetalShaderFormat/Private/MetalShaderFormat.cpp
+++ b/Engine/Source/Developer/Apple/MetalShaderFormat/Private/MetalShaderFormat.cpp
@@ -13,6 +13,7 @@
 #include "HAL/FileManager.h"
 #include "HAL/PlatformFileManager.h"
 #include "Serialization/Archive.h"
+#include "Misc/CommandLine.h"
 #include "Misc/ConfigCacheIni.h"
 #include "MetalBackend.h"
 #include "Misc/FileHelper.h"
diff --git a/Engine/Source/Developer/FunctionalTesting/Private/FunctionalTest.cpp b/Engine/Source/Developer/FunctionalTesting/Private/FunctionalTest.cpp
--- a/Engine/Source/Developer/FunctionalTesting/Private/FunctionalTest.cpp
+++ b/Engine/Source/Developer/FunctionalTesting/Private/FunctionalTest.cpp
@@ -32,6 +32,7 @@
 #include "Misc/RuntimeErrors.h"
 #include "FunctionalTestBase.h"
 #include "UnrealClient.h"
+#include "RenderTimer.h"
 
 #if WITH_EDITOR
 #include "AutomationControllerSettings.h"
diff --git a/Engine/Source/Developer/ShaderCompilerCommon/Public/ShaderMinifier.h b/Engine/Source/Developer/ShaderCompilerCommon/Public/ShaderMinifier.h
--- a/Engine/Source/Developer/ShaderCompilerCommon/Public/ShaderMinifier.h
+++ b/Engine/Source/Developer/ShaderCompilerCommon/Public/ShaderMinifier.h
@@ -4,6 +4,7 @@
 
 #include "CoreMinimal.h"
 
+#include "Containers/AnsiString.h"
 #include "Containers/UnrealString.h"
 #include "Misc/EnumClassFlags.h"
 #include "ShaderSource.h"
diff --git a/Engine/Source/Developer/ShaderFormatOpenGL/Private/OpenGLShaderCompiler.cpp b/Engine/Source/Developer/ShaderFormatOpenGL/Private/OpenGLShaderCompiler.cpp
--- a/Engine/Source/Developer/ShaderFormatOpenGL/Private/OpenGLShaderCompiler.cpp
+++ b/Engine/Source/Developer/ShaderFormatOpenGL/Private/OpenGLShaderCompiler.cpp
@@ -1,6 +1,7 @@
 // Copyright Epic Games, Inc. All Rights Reserved.
 
 #include "CoreMinimal.h"
+#include "Containers/AnsiString.h"
 #include "CrossCompiler.h"
 #include "HAL/FileManager.h"
 #include "HlslccHeaderWriter.h"
diff --git a/Engine/Source/Editor/UnrealEd/Private/Commandlets/CompileShadersTestBedCommandlet.cpp b/Engine/Source/Editor/UnrealEd/Private/Commandlets/CompileShadersTestBedCommandlet.cpp
--- a/Engine/Source/Editor/UnrealEd/Private/Commandlets/CompileShadersTestBedCommandlet.cpp
+++ b/Engine/Source/Editor/UnrealEd/Private/Commandlets/CompileShadersTestBedCommandlet.cpp
@@ -14,6 +14,7 @@
 #include "Materials/Material.h"
 #include "Materials/MaterialInstance.h"
 #include "Materials/MaterialInstanceConstant.h"
+#include "RenderingThread.h"
 #include "ShaderCompiler.h"
 
 #include UE_INLINE_GENERATED_CPP_BY_NAME(CompileShadersTestBedCommandlet)
diff --git a/Engine/Source/Editor/UnrealEd/Private/UnrealEdMisc.cpp b/Engine/Source/Editor/UnrealEd/Private/UnrealEdMisc.cpp
--- a/Engine/Source/Editor/UnrealEd/Private/UnrealEdMisc.cpp
+++ b/Engine/Source/Editor/UnrealEd/Private/UnrealEdMisc.cpp
@@ -109,6 +109,7 @@
 #include "AI/NavigationSystemBase.h"
 #include "Editor/EditorPerProjectUserSettings.h"
 #include "Async/TaskGraphInterfaces.h"
+#include "RenderTimer.h"
 
 #if PLATFORM_MAC
 #include "Mac/MacSystemIncludes.h"
diff --git a/Engine/Source/Runtime/D3D12RHI/Private/D3D12Adapter.cpp b/Engine/Source/Runtime/D3D12RHI/Private/D3D12Adapter.cpp
--- a/Engine/Source/Runtime/D3D12RHI/Private/D3D12Adapter.cpp
+++ b/Engine/Source/Runtime/D3D12RHI/Private/D3D12Adapter.cpp
@@ -12,6 +12,7 @@ D3D12Adapter.cpp:D3D12 Adapter implementation.
 #include "Misc/CommandLine.h"
 #include "Misc/EngineVersion.h"
 #include "Misc/OutputDeviceRedirector.h"
+#include "RenderingThread.h"
 #include "ShaderDiagnostics.h"
 #include "DataDrivenShaderPlatformInfo.h"
 #include "Windows/AllowWindowsPlatformTypes.h"
diff --git a/Engine/Source/Runtime/D3D12RHI/Private/D3D12View.h b/Engine/Source/Runtime/D3D12RHI/Private/D3D12View.h
--- a/Engine/Source/Runtime/D3D12RHI/Private/D3D12View.h
+++ b/Engine/Source/Runtime/D3D12RHI/Private/D3D12View.h
@@ -6,6 +6,7 @@
 #include "D3D12RHICommon.h"
 #include "D3D12Descriptors.h"
 #include "D3D12Resources.h"
+#include "RHICommandList.h"
 #include "RHIResources.h"
 
 class FD3D12Buffer;
diff --git a/Engine/Source/Runtime/D3D12RHI/Private/Windows/WindowsD3D12Viewport.cpp b/Engine/Source/Runtime/D3D12RHI/Private/Windows/WindowsD3D12Viewport.cpp
--- a/Engine/Source/Runtime/D3D12RHI/Private/Windows/WindowsD3D12Viewport.cpp
+++ b/Engine/Source/Runtime/D3D12RHI/Private/Windows/WindowsD3D12Viewport.cpp
@@ -7,6 +7,7 @@
 #include "D3D12RHIPrivate.h"
 #include "Features/IModularFeatures.h"
 #include "HDRHelper.h"
+#include "RenderingThread.h"
 #include "RHIUtilities.h"
 #include "HAL/ThreadHeartBeat.h"
 #include "Windows/IDXGISwapchainProvider.h"
diff --git a/Engine/Source/Runtime/Engine/Private/Animation/AnimSequenceTransformProviderCompiler.cpp b/Engine/Source/Runtime/Engine/Private/Animation/AnimSequenceTransformProviderCompiler.cpp
--- a/Engine/Source/Runtime/Engine/Private/Animation/AnimSequenceTransformProviderCompiler.cpp
+++ b/Engine/Source/Runtime/Engine/Private/Animation/AnimSequenceTransformProviderCompiler.cpp
@@ -25,6 +25,7 @@
 #include "ProfilingDebugging/CountersTrace.h"
 #include "Interfaces/ITargetPlatform.h"
 #include "Interfaces/ITargetPlatformManagerModule.h"
+#include "Serialization/MemoryHasher.h"
 
 #define LOCTEXT_NAMESPACE "AnimSequenceTransformProviderCompiler"
 
diff --git a/Engine/Source/Runtime/Engine/Private/ChartCreation.cpp b/Engine/Source/Runtime/Engine/Private/ChartCreation.cpp
--- a/Engine/Source/Runtime/Engine/Private/ChartCreation.cpp
+++ b/Engine/Source/Runtime/Engine/Private/ChartCreation.cpp
@@ -22,6 +22,7 @@
 #include "Stats/Stats.h"
 #include "UnrealClient.h"
 #include "LegacyScreenPercentageDriver.h"
+#include "RenderTimer.h"
 
 #ifndef FPS_CHART_SUPPORT_CSV_PROFILE
 #define FPS_CHART_SUPPORT_CSV_PROFILE (CSV_PROFILER && !UE_BUILD_SHIPPING)
diff --git a/Engine/Source/Runtime/Engine/Private/Rendering/ColorVertexBuffer.cpp b/Engine/Source/Runtime/Engine/Private/Rendering/ColorVertexBuffer.cpp
--- a/Engine/Source/Runtime/Engine/Private/Rendering/ColorVertexBuffer.cpp
+++ b/Engine/Source/Runtime/Engine/Private/Rendering/ColorVertexBuffer.cpp
@@ -9,6 +9,7 @@
 #include "VertexFactory.h"
 #include "DataDrivenShaderPlatformInfo.h"
 #include "GlobalRenderResources.h"
+#include "RHICommandList.h"
 #include "RHIResourceReplace.h"
 
 /*-----------------------------------------------------------------------------
diff --git a/Engine/Source/Runtime/Engine/Private/ShaderCompiler/ShaderCompilerStats.cpp b/Engine/Source/Runtime/Engine/Private/ShaderCompiler/ShaderCompilerStats.cpp
--- a/Engine/Source/Runtime/Engine/Private/ShaderCompiler/ShaderCompilerStats.cpp
+++ b/Engine/Source/Runtime/Engine/Private/ShaderCompiler/ShaderCompilerStats.cpp
@@ -17,6 +17,7 @@
 #include "ProfilingDebugging/DiagnosticTable.h"
 #include "ProfilingDebugging/LoadTimeTracker.h"
 #include "ProfilingDebugging/StallDetector.h"
+#include "RenderingThread.h"
 #include "Serialization/CompactBinaryWriter.h"
 #include "ShaderCompileWorkerUtil.h"
 #include "Misc/FileHelper.h"
diff --git a/Engine/Source/Runtime/Engine/Private/UnrealClient.cpp b/Engine/Source/Runtime/Engine/Private/UnrealClient.cpp
--- a/Engine/Source/Runtime/Engine/Private/UnrealClient.cpp
+++ b/Engine/Source/Runtime/Engine/Private/UnrealClient.cpp
@@ -35,6 +35,7 @@
 #include "Stats/ThreadIdleStats.h"
 #include "RHIUtilities.h"
 #include "Engine/CoreSettings.h"
+#include "RenderTimer.h"
 
 DEFINE_LOG_CATEGORY_STATIC(LogClient, Log, All);
 
diff --git a/Engine/Source/Runtime/Engine/Public/LocalVertexFactory.h b/Engine/Source/Runtime/Engine/Public/LocalVertexFactory.h
--- a/Engine/Source/Runtime/Engine/Public/LocalVertexFactory.h
+++ b/Engine/Source/Runtime/Engine/Public/LocalVertexFactory.h
@@ -4,6 +4,7 @@
 
 #include "CoreMinimal.h"
 #include "ShaderParameters.h"
+#include "ShaderParameterMacros.h"
 #include "Components.h"
 #include "VertexFactory.h"
 #include "GlobalRenderResources.h"
diff --git a/Engine/Source/Runtime/Launch/Private/LaunchEngineLoop.cpp b/Engine/Source/Runtime/Launch/Private/LaunchEngineLoop.cpp
--- a/Engine/Source/Runtime/Launch/Private/LaunchEngineLoop.cpp
+++ b/Engine/Source/Runtime/Launch/Private/LaunchEngineLoop.cpp
@@ -54,6 +54,7 @@
 #include "ProfilingDebugging/BootProfiling.h"
 #if WITH_ENGINE
 #include "HAL/PlatformSplash.h"
+#include "RenderTimer.h"
 #include "SceneInterface.h"
 #include "StereoRenderUtils.h"
 #include "DataDrivenShaderPlatformInfo.h"
diff --git a/Engine/Source/Runtime/OpenGLDrv/Private/OpenGLDevice.cpp b/Engine/Source/Runtime/OpenGLDrv/Private/OpenGLDevice.cpp
--- a/Engine/Source/Runtime/OpenGLDrv/Private/OpenGLDevice.cpp
+++ b/Engine/Source/Runtime/OpenGLDrv/Private/OpenGLDevice.cpp
@@ -15,6 +15,7 @@
 #include "RHI.h"
 #include "DataDrivenShaderPlatformInfo.h"
 #include "Containers/List.h"
+#include "RenderingThread.h"
 #include "RenderResource.h"
 #include "ShaderCore.h"
 #include "Shader.h"
diff --git a/Engine/Source/Runtime/OpenGLDrv/Private/OpenGLViewport.cpp b/Engine/Source/Runtime/OpenGLDrv/Private/OpenGLViewport.cpp
--- a/Engine/Source/Runtime/OpenGLDrv/Private/OpenGLViewport.cpp
+++ b/Engine/Source/Runtime/OpenGLDrv/Private/OpenGLViewport.cpp
@@ -7,6 +7,7 @@
 #include "CoreMinimal.h"
 #include "Stats/Stats.h"
 #include "HAL/IConsoleManager.h"
+#include "RenderingThread.h"
 #include "RHI.h"
 #include "OpenGLDrv.h"
 #include "OpenGLDrvPrivate.h"
diff --git a/Engine/Source/Runtime/PreLoadScreen/Private/PreLoadScreenManager.cpp b/Engine/Source/Runtime/PreLoadScreen/Private/PreLoadScreenManager.cpp
--- a/Engine/Source/Runtime/PreLoadScreen/Private/PreLoadScreenManager.cpp
+++ b/Engine/Source/Runtime/PreLoadScreen/Private/PreLoadScreenManager.cpp
@@ -5,6 +5,7 @@
 #include "Engine/GameEngine.h"
 #include "Framework/Application/SlateApplication.h"
 #include "GlobalShader.h"
+#include "RenderingThread.h"
 #include "ShaderCompiler.h"
 
 #include "PreLoadScreen.h"
diff --git a/Engine/Source/Runtime/Renderer/Private/RayTracing/RayTracingBasicShaders.cpp b/Engine/Source/Runtime/Renderer/Private/RayTracing/RayTracingBasicShaders.cpp
--- a/Engine/Source/Runtime/Renderer/Private/RayTracing/RayTracingBasicShaders.cpp
+++ b/Engine/Source/Runtime/Renderer/Private/RayTracing/RayTracingBasicShaders.cpp
@@ -1,6 +1,6 @@
 // Copyright Epic Games, Inc. All Rights Reserved.
 
-#include "RayTracingBasicShaders.h"
+#include "RayTracing/RayTracingBasicShaders.h"
 
 #if RHI_RAYTRACING
 
diff --git a/Engine/Source/Runtime/Renderer/Private/RayTracing/RayTracingBasicShaders.h b/Engine/Source/Runtime/Renderer/Private/RayTracing/RayTracingBasicShaders.h
--- a/Engine/Source/Runtime/Renderer/Private/RayTracing/RayTracingBasicShaders.h
+++ b/Engine/Source/Runtime/Renderer/Private/RayTracing/RayTracingBasicShaders.h
@@ -7,6 +7,7 @@
 #if RHI_RAYTRACING
 
 #include "RayTracingPayloadType.h"
+#include "ShaderParameterStruct.h"
 
 // C++ counter-part of FBasicRayData declared in RayTracingCommon.ush
 struct FBasicRayTracingRay
diff --git a/Engine/Source/Runtime/Renderer/Public/MeshDrawShaderBindings.h b/Engine/Source/Runtime/Renderer/Public/MeshDrawShaderBindings.h
--- a/Engine/Source/Runtime/Renderer/Public/MeshDrawShaderBindings.h
+++ b/Engine/Source/Runtime/Renderer/Public/MeshDrawShaderBindings.h
@@ -7,6 +7,7 @@
 #pragma once
 
 #include "Shader.h"
+#include "UniformBuffer.h"
 
 // Whether to assert when mesh command shader bindings were not set by the pass processor.
 // Enabled by default in debug
diff --git a/Engine/Source/Runtime/SlateRHIRenderer/Private/SlateRHIRenderer.cpp b/Engine/Source/Runtime/SlateRHIRenderer/Private/SlateRHIRenderer.cpp
--- a/Engine/Source/Runtime/SlateRHIRenderer/Private/SlateRHIRenderer.cpp
+++ b/Engine/Source/Runtime/SlateRHIRenderer/Private/SlateRHIRenderer.cpp
@@ -29,6 +29,7 @@
 #include "Rendering/ElementBatcher.h"
 #include "Rendering/SlateRenderer.h"
 #include "RenderResource.h"
+#include "RenderTimer.h"
 #include "RenderingThread.h"
 #include "RHIResources.h"
 #include "RHIUtilities.h"
diff --git a/Engine/Source/Runtime/VulkanRHI/Private/VulkanPipeline.cpp b/Engine/Source/Runtime/VulkanRHI/Private/VulkanPipeline.cpp
--- a/Engine/Source/Runtime/VulkanRHI/Private/VulkanPipeline.cpp
+++ b/Engine/Source/Runtime/VulkanRHI/Private/VulkanPipeline.cpp
@@ -10,6 +10,7 @@
 #include "HAL/FileManager.h"
 #include "Misc/Paths.h"
 #include "Misc/FileHelper.h"
+#include "RenderingThread.h"
 #include "Serialization/MemoryReader.h"
 #include "Serialization/MemoryWriter.h"
 #include "VulkanPendingState.h"
diff --git a/Engine/Source/Runtime/VulkanRHI/Private/VulkanRHI.cpp b/Engine/Source/Runtime/VulkanRHI/Private/VulkanRHI.cpp
--- a/Engine/Source/Runtime/VulkanRHI/Private/VulkanRHI.cpp
+++ b/Engine/Source/Runtime/VulkanRHI/Private/VulkanRHI.cpp
@@ -20,6 +20,7 @@
 #include "VulkanLLM.h"
 #include "Misc/EngineVersion.h"
 #include "GlobalShader.h"
+#include "RenderingThread.h"
 #include "RHICore.h"
 #include "VulkanResourceCollection.h"
 #include "RHIValidation.h"
diff --git a/Engine/Source/Runtime/VulkanRHI/Private/VulkanRayTracing.cpp b/Engine/Source/Runtime/VulkanRHI/Private/VulkanRayTracing.cpp
--- a/Engine/Source/Runtime/VulkanRHI/Private/VulkanRayTracing.cpp
+++ b/Engine/Source/Runtime/VulkanRHI/Private/VulkanRayTracing.cpp
@@ -10,6 +10,7 @@
 #include "Async/ParallelFor.h"
 #include "Misc/CommandLine.h"
 #include "Misc/Parse.h"
+#include "ShaderParameterMacros.h"
 
 static int32 GVulkanRayTracingAllowCompaction = 1;
 static FAutoConsoleVariableRef CVarVulkanRayTracingAllowCompaction(
diff --git a/Engine/Source/Runtime/VulkanRHI/Private/VulkanResources.h b/Engine/Source/Runtime/VulkanRHI/Private/VulkanResources.h
--- a/Engine/Source/Runtime/VulkanRHI/Private/VulkanResources.h
+++ b/Engine/Source/Runtime/VulkanRHI/Private/VulkanResources.h
@@ -7,6 +7,7 @@
 #pragma once
 
 #include "CoreMinimal.h"
+#include "RHICommandList.h"
 #include "VulkanPlatform.h"
 #include "VulkanConfiguration.h"
 #include "VulkanState.h"
diff --git a/Engine/Source/Runtime/Windows/D3D11RHI/Private/D3D11Device.cpp b/Engine/Source/Runtime/Windows/D3D11RHI/Private/D3D11Device.cpp
--- a/Engine/Source/Runtime/Windows/D3D11RHI/Private/D3D11Device.cpp
+++ b/Engine/Source/Runtime/Windows/D3D11RHI/Private/D3D11Device.cpp
@@ -9,6 +9,8 @@
 #include "Misc/CommandLine.h"
 #include "Misc/ConfigCacheIni.h"
 #include "Modules/ModuleManager.h"
+#include "RenderingThread.h"
+
 #include "Windows/AllowWindowsPlatformTypes.h"
 	#include <delayimp.h>
 	#if WITH_AMD_AGS
diff --git a/Engine/Source/Runtime/Windows/D3D11RHI/Private/Windows/WindowsD3D11Device.cpp b/Engine/Source/Runtime/Windows/D3D11RHI/Private/Windows/WindowsD3D11Device.cpp
--- a/Engine/Source/Runtime/Windows/D3D11RHI/Private/Windows/WindowsD3D11Device.cpp
+++ b/Engine/Source/Runtime/Windows/D3D11RHI/Private/Windows/WindowsD3D11Device.cpp
@@ -12,6 +12,7 @@
 #include "IHeadMountedDisplayModule.h"
 #include "GenericPlatform/GenericPlatformDriver.h"			// FGPUDriverInfo
 #include "GenericPlatform/GenericPlatformCrashContext.h"
+#include "RenderingThread.h"
 #include "RHIValidation.h"
 #include "RHIUtilities.h"
 #include "ShaderDiagnostics.h"
diff --git a/Engine/Source/Runtime/Windows/D3D11RHI/Private/Windows/WindowsD3D11Viewport.cpp b/Engine/Source/Runtime/Windows/D3D11RHI/Private/Windows/WindowsD3D11Viewport.cpp
--- a/Engine/Source/Runtime/Windows/D3D11RHI/Private/Windows/WindowsD3D11Viewport.cpp
+++ b/Engine/Source/Runtime/Windows/D3D11RHI/Private/Windows/WindowsD3D11Viewport.cpp
@@ -8,6 +8,7 @@
 #include "Misc/CommandLine.h"
 #include "Windows/IDXGISwapchainProvider.h"
 #include "Features/IModularFeatures.h"
+#include "RenderingThread.h"
 #include "RHIUtilities.h"
 
 #include "Windows/AllowWindowsPlatformTypes.h"
